#ttp 322654 - UE4: RENDERING: Texture compression pipeline should support per-format maximum sizes
#proj Engine
#branch UE4
#summary Added GetFormatCapabilities to ITextureFormat interface, TextureCompressorModule attempts to downsample images that are too large.
#add GetFormatCapabilities to ITextureFormat and all affected subclasses.
#change BuildTextureMips now attempts to construct a source that is compatible with the maximum size for a texture compressor.

diff --git a/Engine/Source/Developer/TargetPlatform/Public/Interfaces/ITextureFormat.h b/Engine/Source/Developer/TargetPlatform/Public/Interfaces/ITextureFormat.h
--- a/Engine/Source/Developer/TargetPlatform/Public/Interfaces/ITextureFormat.h
+++ b/Engine/Source/Developer/TargetPlatform/Public/Interfaces/ITextureFormat.h
@@ -6,6 +6,14 @@
 
 #pragma once
 
+struct FTextureFormatCompressorCaps
+{
+	FTextureFormatCompressorCaps() : MaxTextureDimension(TNumericLimits<uint32>::Max())
+	{
+	}
+
+	uint32 MaxTextureDimension;
+};
 
 /**
  * Interface for texture compression modules.
@@ -53,6 +61,15 @@ public:
 	 */
 	virtual void GetSupportedFormats(TArray<FName>& OutFormats) const = 0;
 
+
+
+	/**
+	* Gets the capabilities of the texture compressor.
+	*
+	* @param OutCaps - filled with capability properties of texture format compressor.
+	*/
+	virtual FTextureFormatCompressorCaps GetFormatCapabilities() const = 0;
+
 	/**
 	 * Compresses a single image.
 	 *
diff --git a/Engine/Source/Developer/TextureCompressor/Private/TextureCompressorModule.cpp b/Engine/Source/Developer/TextureCompressor/Private/TextureCompressorModule.cpp
--- a/Engine/Source/Developer/TextureCompressor/Private/TextureCompressorModule.cpp
+++ b/Engine/Source/Developer/TextureCompressor/Private/TextureCompressorModule.cpp
@@ -521,17 +521,20 @@ static void GenerateTopMip(const FImage& SrcImage, FImage& DestImage, const FTex
 /**
  * Generate a full mip chain. The input mip chain must have one or more mips.
  * @param Settings - Preprocess settings.
- * @param OutMipChain - The full mip chain.
+ * @param BaseImage - An image that will serve as the source for the generation of the mip chain.
+ * @param OutMipChain - An array that will contain the resultant mip images. Generated mip levels are appended to the array.
+ * @param MipChainDepth - number of mip images to produce. Mips chain is finished when either a 1x1 mip is produced or 'MipChainDepth' images have been produced.
  */
 static void GenerateMipChain(
 	const FTextureBuildSettings& Settings,
-	TArray<FImage> &OutMipChain
+	const FImage& BaseImage,
+	TArray<FImage> &OutMipChain,
+	uint32 MipChainDepth = MAX_uint32
 	)
 {
-	check(OutMipChain.Num() >= 1);
-	check(OutMipChain[0].Format == ERawImageFormat::RGBA32F);
+	check(BaseImage.Format == ERawImageFormat::RGBA32F);
 
-	const FImage& BaseMip = OutMipChain.Last();
+	const FImage& BaseMip = BaseImage;
 	const int32 SrcWidth = BaseMip.SizeX;
 	const int32 SrcHeight= BaseMip.SizeY;
 	const int32 SrcNumSlices = BaseMip.NumSlices;
@@ -558,7 +561,7 @@ static void GenerateMipChain(
 	}
 
 	// Generate mips
-	for (;;)
+	for (; MipChainDepth != 0 ; --MipChainDepth)
 	{
 		FImage& DestImage = *new(OutMipChain) FImage(IntermediateDst.SizeX, IntermediateDst.SizeY, SrcNumSlices, ImageFormat);
 		
@@ -1472,6 +1475,21 @@ private:
 };
 typedef FAsyncTask<FAsyncCompressionWorker> FAsyncCompressionTask;
 
+FTextureFormatCompressorCaps GetTextureFormatCaps(const FTextureBuildSettings& Settings)
+{
+	ITargetPlatformManagerModule* TPM = GetTargetPlatformManager();
+	if (TPM)
+	{
+		const ITextureFormat* TextureFormat = TPM->FindTextureFormat(Settings.TextureFormatName);
+		if (TextureFormat != nullptr)
+		{
+			return TextureFormat->GetFormatCapabilities();
+		}
+	}
+	
+	return FTextureFormatCompressorCaps();
+}
+
 // compress mip-maps in InMipChain and add mips to Texture, might alter the source content
 static bool CompressMipChain(
 	const TArray<FImage>& MipChain,
@@ -1665,46 +1683,93 @@ private:
 	HMODULE nvTextureToolsHandle;
 #endif	//PLATFORM_WINDOWS
 
-	
 	bool BuildTextureMips(
-		const TArray<FImage>& SourceMips,
+		const TArray<FImage>& InSourceMips,
 		const FTextureBuildSettings& BuildSettings,
 		TArray<FImage>& OutMipChain)
 	{
-		check(SourceMips.Num());
-		check(SourceMips[0].SizeX > 0 && SourceMips[0].SizeY > 0 && SourceMips[0].NumSlices > 0);
+		check(InSourceMips.Num());
+		check(InSourceMips[0].SizeX > 0 && InSourceMips[0].SizeY > 0 && InSourceMips[0].NumSlices > 0);
+		const FTextureFormatCompressorCaps CompressorCaps = GetTextureFormatCaps(BuildSettings);
 
 		// Identify long-lat cubemaps.
-		bool bLongLatCubemap = BuildSettings.bCubemap && SourceMips[0].NumSlices == 1;
+		bool bLongLatCubemap = BuildSettings.bCubemap && InSourceMips[0].NumSlices == 1;
 
-		if (BuildSettings.bCubemap && SourceMips[0].NumSlices != 6 && !bLongLatCubemap)
+		if (BuildSettings.bCubemap && InSourceMips[0].NumSlices != 6 && !bLongLatCubemap)
 		{
 			return false;
 		}
 
-		// Figure out how many source mips to use.
-		int32 NumTotalMips = (BuildSettings.MipGenSettings == TMGS_NoMipmaps)
-			? 1
-			: 1 + FMath::Max(FMath::CeilLogTwo(SourceMips[0].SizeX), FMath::CeilLogTwo(SourceMips[0].SizeY));
-		if (bLongLatCubemap)
-		{
-			NumTotalMips = 1 + FMath::CeilLogTwo(ComputeLongLatCubemapExtents(SourceMips[0]));
-		}
+		// Determine the maximum possible mip counts for source and dest.
+		const int32 MaxSourceMipCount = bLongLatCubemap ?
+			1 + FMath::CeilLogTwo(ComputeLongLatCubemapExtents(InSourceMips[0])) :
+			1 + FMath::CeilLogTwo(FMath::Max(InSourceMips[0].SizeX, InSourceMips[0].SizeY));
+		const int32 MaxDestMipCount = 1 + FMath::CeilLogTwo(CompressorCaps.MaxTextureDimension);
+
+		// Determine the number of mips required by BuildSettings.
+		int32 NumOutputMips = (BuildSettings.MipGenSettings == TMGS_NoMipmaps) ? 1 : MaxSourceMipCount;
+		NumOutputMips = FMath::Min(NumOutputMips, MaxDestMipCount);
+
+		int32 NumSourceMips = InSourceMips.Num();
 
-		int32 NumSourceMips = SourceMips.Num();
 		if (BuildSettings.MipGenSettings != TMGS_LeaveExistingMips ||
 			BuildSettings.MipGenSettings == TMGS_NoMipmaps ||
 			bLongLatCubemap)
 		{
 			NumSourceMips = 1;
 		}
-		OutMipChain.Empty(NumTotalMips);
 
-		// Figure out how many slices we are generating mipchains for.
-		int32 NumTotalSlices = BuildSettings.bCubemap ? 6 : 1;
+		// Check our source can be exported by our current compressor.
+		int32 LevelsToUsableSource = FMath::Max(0, MaxSourceMipCount - MaxDestMipCount);
+		int32 StartMip = FMath::Max(0, LevelsToUsableSource);
+		bool bBuildSourceImage = StartMip > (NumSourceMips - 1);
+
+		TArray<FImage> GeneratedSourceMips;
+		if (bBuildSourceImage)
+		{			
+			// the source is larger than the compressor allows and no mip image exists to act as a smaller source.
+			// We must generate a suitable source image:
+			bool bSuitableFormat = InSourceMips.Last().Format == ERawImageFormat::RGBA32F;
+			const FImage& BaseImage = InSourceMips.Last();
+
+			if (BaseImage.SizeX != FMath::RoundUpToPowerOfTwo(BaseImage.SizeX) || BaseImage.SizeY != FMath::RoundUpToPowerOfTwo(BaseImage.SizeY))
+			{
+				UE_LOG(LogTextureCompressor, Warning,
+					TEXT("Source image %dx%d (npot) prevents resizing and is too large for compressors max dimension (%d)."),
+					BaseImage.SizeX,
+					BaseImage.SizeY,
+					CompressorCaps.MaxTextureDimension
+					);
+				return false;
+			}
+
+			FImage Temp;
+			if (!bSuitableFormat)
+			{
+				// convert to RGBA32F
+				BaseImage.CopyTo(Temp, ERawImageFormat::RGBA32F, false);
+			}
+
+			UE_LOG(LogTextureCompressor, Verbose,
+				TEXT("Source image %dx%d too large for compressors max dimension (%d). Resizing."),
+				BaseImage.SizeX,
+				BaseImage.SizeY,
+				CompressorCaps.MaxTextureDimension
+				);
+			GenerateMipChain(BuildSettings, bSuitableFormat ? BaseImage : Temp, GeneratedSourceMips, LevelsToUsableSource);
+
+			check(GeneratedSourceMips.Num() != 0);
+			// Note: The newly generated mip chain does not include the original top level mip.
+			StartMip--;
+		}
+
+		const TArray<FImage>& SourceMips = bBuildSourceImage ? GeneratedSourceMips : InSourceMips;
 
+		OutMipChain.Empty(NumOutputMips);
 		// Copy over base mips.
-		for (int32 MipIndex = 0; MipIndex < NumSourceMips; ++MipIndex)
+		int32 CopyCount = FMath::Min(NumOutputMips - StartMip, SourceMips.Num());
+		check(StartMip <= SourceMips.Num());
+		for (int32 MipIndex = StartMip; MipIndex < CopyCount; ++MipIndex)
 		{
 			const FImage& Image = SourceMips[MipIndex];
 			const int32 SrcWidth = Image.SizeX;
@@ -1760,19 +1825,19 @@ private:
 		}
 
 		// Generate any missing mips in the chain.
-		if (NumTotalMips > OutMipChain.Num())
+		if (NumOutputMips > OutMipChain.Num())
 		{
 			// Do angular filtering of cubemaps if requested.
 			if (BuildSettings.bCubemap)
 			{
-				GenerateAngularFilteredMips(OutMipChain, NumTotalMips, BuildSettings.DiffuseConvolveMipLevel);
+				GenerateAngularFilteredMips(OutMipChain, NumOutputMips, BuildSettings.DiffuseConvolveMipLevel);
 			}
 			else
 			{
-				GenerateMipChain(BuildSettings, OutMipChain);
+				GenerateMipChain(BuildSettings, OutMipChain.Last(), OutMipChain);
 			}
 		}
-		check(OutMipChain.Num() == NumTotalMips);
+		check(OutMipChain.Num() == NumOutputMips);
 
 		// Apply post-mip generation adjustments.
 		if (BuildSettings.bReplicateRed)
diff --git a/Engine/Source/Developer/TextureFormatAndroid/Private/TextureFormatAndroid.cpp b/Engine/Source/Developer/TextureFormatAndroid/Private/TextureFormatAndroid.cpp
--- a/Engine/Source/Developer/TextureFormatAndroid/Private/TextureFormatAndroid.cpp
+++ b/Engine/Source/Developer/TextureFormatAndroid/Private/TextureFormatAndroid.cpp
@@ -141,6 +141,11 @@ class FTextureFormatAndroid : public ITextureFormat
 		}
 	}
 
+	virtual FTextureFormatCompressorCaps GetFormatCapabilities() const OVERRIDE
+	{
+		return FTextureFormatCompressorCaps(); // Default capabilities.
+	}
+
 	virtual bool CompressImage(
 		const FImage& InImage,
 		const struct FTextureBuildSettings& BuildSettings,
diff --git a/Engine/Source/Developer/TextureFormatDXT/Private/TextureFormatDXT.cpp b/Engine/Source/Developer/TextureFormatDXT/Private/TextureFormatDXT.cpp
--- a/Engine/Source/Developer/TextureFormatDXT/Private/TextureFormatDXT.cpp
+++ b/Engine/Source/Developer/TextureFormatDXT/Private/TextureFormatDXT.cpp
@@ -404,6 +404,11 @@ class FTextureFormatDXT : public ITextureFormat
 			OutFormats.Add(GSupportedTextureFormatNames[i]);
 		}
 	}
+	
+	virtual FTextureFormatCompressorCaps GetFormatCapabilities() const OVERRIDE
+	{
+		return FTextureFormatCompressorCaps(); // Default capabilities.
+	}
 
 	virtual bool CompressImage(
 		const FImage& InImage,
diff --git a/Engine/Source/Developer/TextureFormatPVR/Private/TextureFormatPVR.cpp b/Engine/Source/Developer/TextureFormatPVR/Private/TextureFormatPVR.cpp
--- a/Engine/Source/Developer/TextureFormatPVR/Private/TextureFormatPVR.cpp
+++ b/Engine/Source/Developer/TextureFormatPVR/Private/TextureFormatPVR.cpp
@@ -245,6 +245,14 @@ class FTextureFormatPVR : public ITextureFormat
 		}
 	}
 
+	virtual FTextureFormatCompressorCaps GetFormatCapabilities() const OVERRIDE
+	{
+		FTextureFormatCompressorCaps RetCaps;
+		// PVR compressor is limited to <=4096 in any direction.
+		RetCaps.MaxTextureDimension = 4096;
+		return RetCaps;
+	}
+
 	virtual bool CompressImage(
 		const FImage& InImage,
 		const struct FTextureBuildSettings& BuildSettings,
diff --git a/Engine/Source/Developer/TextureFormatUncompressed/Private/TextureFormatUncompressed.cpp b/Engine/Source/Developer/TextureFormatUncompressed/Private/TextureFormatUncompressed.cpp
--- a/Engine/Source/Developer/TextureFormatUncompressed/Private/TextureFormatUncompressed.cpp
+++ b/Engine/Source/Developer/TextureFormatUncompressed/Private/TextureFormatUncompressed.cpp
@@ -55,6 +55,11 @@ class FTextureFormatUncompressed : public ITextureFormat
 			OutFormats.Add(GSupportedTextureFormatNames[i]);
 		}
 	}
+	
+	virtual FTextureFormatCompressorCaps GetFormatCapabilities() const OVERRIDE
+	{
+		return FTextureFormatCompressorCaps(); // Default capabilities.
+	}
 
 	virtual bool CompressImage(
 		const FImage& InImage,
