* fixed TTP 340094 FIXIF: CRITICAL: ElementalDemo: XboxOne: Some particle corruption
 also get CL2182557

diff --git a/Engine/Source/Developer/TextureCompressor/Private/TextureCompressorModule.cpp b/Engine/Source/Developer/TextureCompressor/Private/TextureCompressorModule.cpp
--- a/Engine/Source/Developer/TextureCompressor/Private/TextureCompressorModule.cpp
+++ b/Engine/Source/Developer/TextureCompressor/Private/TextureCompressorModule.cpp
@@ -1651,8 +1651,8 @@ public:
 
 		// Set the correct biased texture size so that the compressor understands the original source image size
 		// This is requires for platforms that may need to tile based on the original source texture size
-		BuildSettings.OriginalTextureSize.X = IntermediateMipChain[0].SizeX;
-		BuildSettings.OriginalTextureSize.Y = IntermediateMipChain[0].SizeY;
+		BuildSettings.TopMipSize.X = IntermediateMipChain[0].SizeX;
+		BuildSettings.TopMipSize.Y = IntermediateMipChain[0].SizeY;
 
 		return CompressMipChain(IntermediateMipChain, BuildSettings, OutTextureMips);
 	}
diff --git a/Engine/Source/Developer/TextureCompressor/Public/TextureCompressorModule.h b/Engine/Source/Developer/TextureCompressor/Public/TextureCompressorModule.h
--- a/Engine/Source/Developer/TextureCompressor/Public/TextureCompressorModule.h
+++ b/Engine/Source/Developer/TextureCompressor/Public/TextureCompressorModule.h
@@ -107,8 +107,8 @@ struct FTextureBuildSettings
 	float CompositePower;
 	/** The source texture's final LOD bias (i.e. includes LODGroup based biases) */
 	uint32 LODBias;
-	/** The texture's size, mutable as this is filled in once the mips has been built */
-	mutable FIntPoint OriginalTextureSize;
+	/** The texture's top mip size without LODBias applied, should be moved into a separate struct together with bImageHasAlphaChannel */
+	mutable FIntPoint TopMipSize;
 	/** Can the texture be streamed */
 	uint32 bStreamable : 1;
 
@@ -134,7 +134,7 @@ struct FTextureBuildSettings
 		, CompositeTextureMode( 0 /*CTM_Disabled*/ )
 		, CompositePower( 1.0f )
 		, LODBias(0)
-		, OriginalTextureSize(0, 0)
+		, TopMipSize(0, 0)
 		, bStreamable(false)
 	{
 	}
diff --git a/Engine/Source/Runtime/Engine/Private/TextureDerivedData.cpp b/Engine/Source/Runtime/Engine/Private/TextureDerivedData.cpp
--- a/Engine/Source/Runtime/Engine/Private/TextureDerivedData.cpp
+++ b/Engine/Source/Runtime/Engine/Private/TextureDerivedData.cpp
@@ -1206,7 +1206,6 @@ static void SerializePlatformData(
 		if (Ar.IsLoading())
 		{
 			check(Texture);
-			Texture->LODBias -= FirstMipToSerialize;
 			FirstMipToSerialize = 0;
 		}
 	}
