Texture : Added Texture property to change maximum resolution limitation (MaxTextureSize), especially for HDR cubemap (which is limited by 512)

diff --git a/Engine/Source/Developer/TextureCompressor/Private/TextureCompressorModule.cpp b/Engine/Source/Developer/TextureCompressor/Private/TextureCompressorModule.cpp
--- a/Engine/Source/Developer/TextureCompressor/Private/TextureCompressorModule.cpp
+++ b/Engine/Source/Developer/TextureCompressor/Private/TextureCompressorModule.cpp
@@ -779,9 +779,9 @@ static FVector ComputeWSCubeDirectionAtTexelCenter(uint32 CubemapFace, uint32 x,
 	return DirectionWS;
 }
 
-static int32 ComputeLongLatCubemapExtents(const FImage& SrcImage)
+static int32 ComputeLongLatCubemapExtents(const FImage& SrcImage, const int32 MaxCubemapTextureResolution)
 {
-	return FMath::Clamp(1 << FMath::FloorLog2(SrcImage.SizeX / 2), 32, 512);
+	return FMath::Clamp(1 << FMath::FloorLog2(SrcImage.SizeX / 2), 32, MaxCubemapTextureResolution);
 }
 
 /**
@@ -789,14 +789,14 @@ static int32 ComputeLongLatCubemapExtents(const FImage& SrcImage)
  * @param OutMip - The output mip.
  * @param SrcImage - The source longlat image.
  */
-static void GenerateBaseCubeMipFromLongitudeLatitude2D(FImage* OutMip, const FImage& SrcImage)
+static void GenerateBaseCubeMipFromLongitudeLatitude2D(FImage* OutMip, const FImage& SrcImage, const int32 MaxCubemapTextureResolution)
 {
 	FImage LongLatImage;
 	SrcImage.CopyTo(LongLatImage, ERawImageFormat::RGBA32F, false);
 	FImageViewLongLat LongLatView(LongLatImage);
 
 	// TODO_TEXTURE: Expose target size to user.
-	int32 Extent = ComputeLongLatCubemapExtents(LongLatImage);
+	int32 Extent = ComputeLongLatCubemapExtents(LongLatImage, MaxCubemapTextureResolution);
 	float InvExtent = 1.0f / Extent;
 	OutMip->Init(Extent, Extent, 6, ERawImageFormat::RGBA32F, false);
 
@@ -1702,9 +1702,9 @@ private:
 
 		// Determine the maximum possible mip counts for source and dest.
 		const int32 MaxSourceMipCount = bLongLatCubemap ?
-			1 + FMath::CeilLogTwo(ComputeLongLatCubemapExtents(InSourceMips[0])) :
+			1 + FMath::CeilLogTwo(ComputeLongLatCubemapExtents(InSourceMips[0], BuildSettings.MaxTextureResolution)) :
 			1 + FMath::CeilLogTwo(FMath::Max(InSourceMips[0].SizeX, InSourceMips[0].SizeY));
-		const int32 MaxDestMipCount = 1 + FMath::CeilLogTwo(CompressorCaps.MaxTextureDimension);
+		const int32 MaxDestMipCount = 1 + FMath::CeilLogTwo(FMath::Min(CompressorCaps.MaxTextureDimension, BuildSettings.MaxTextureResolution));
 
 		// Determine the number of mips required by BuildSettings.
 		int32 NumOutputMips = (BuildSettings.MipGenSettings == TMGS_NoMipmaps) ? 1 : MaxSourceMipCount;
@@ -1782,7 +1782,7 @@ private:
 			if (bLongLatCubemap)
 			{
 				// Generate the base mip from the long-lat source image.
-				GenerateBaseCubeMipFromLongitudeLatitude2D(Mip, Image);
+				GenerateBaseCubeMipFromLongitudeLatitude2D(Mip, Image, BuildSettings.MaxTextureResolution);
 			}
 			else
 			{
diff --git a/Engine/Source/Developer/TextureCompressor/Public/TextureCompressorModule.h b/Engine/Source/Developer/TextureCompressor/Public/TextureCompressorModule.h
--- a/Engine/Source/Developer/TextureCompressor/Public/TextureCompressorModule.h
+++ b/Engine/Source/Developer/TextureCompressor/Public/TextureCompressorModule.h
@@ -71,12 +71,16 @@ struct FTextureBuildSettings
 	uint32 DiffuseConvolveMipLevel;
 	/** The size of the kernel with which mips should be sharpened. 2 for 2x2, 4 for 4x4, 6 for 6x6, 8 for 8x8 */
 	uint32 SharpenMipKernelSize;
+	/** For maximum resolution. */
+	uint32 MaxTextureResolution;
 	/** Format of the compressed texture, used to choose a compression DLL. */
 	FName TextureFormatName;
 	/** Mipmap generation settings. */
 	uint8 MipGenSettings; // TextureMipGenSettings, opaque to avoid dependencies on engine headers.
 	/** Whether the texture being built is a cubemap. */
 	uint32 bCubemap : 1;
+	/** Whether the texture being built from long/lat source to cubemap. */
+	uint32 bLongLatSource : 1;
 	/** Whether the texture contains color data in the sRGB colorspace. */
 	uint32 bSRGB : 1;
 	/** Whether the border of the image should be maintained during mipmap generation. */
@@ -117,8 +121,10 @@ struct FTextureBuildSettings
 		: MipSharpening( 0.0f )
 		, DiffuseConvolveMipLevel( 0 )
 		, SharpenMipKernelSize( 2 )
+		, MaxTextureResolution(TNumericLimits<uint32>::Max())
 		, MipGenSettings( 1 /*TMGS_SimpleAverage*/ )
 		, bCubemap( false )
+		, bLongLatSource(false)
 		, bSRGB( false )
 		, bPreserveBorder( false )
 		, bDitherMipMapAlpha( false )
diff --git a/Engine/Source/Runtime/Engine/Classes/Engine/Texture.h b/Engine/Source/Runtime/Engine/Classes/Engine/Texture.h
--- a/Engine/Source/Runtime/Engine/Classes/Engine/Texture.h
+++ b/Engine/Source/Runtime/Engine/Classes/Engine/Texture.h
@@ -396,6 +396,10 @@ public:
 	UPROPERTY()
 	uint32 CompressionNoMipmaps_DEPRECATED:1;
 
+	/** The maximum resolution for generated textures. A value of 0 means the maximum size for the format on each platform, except HDR long/lat cubemaps, which default to a resolution of 512. */ 
+	UPROPERTY(EditAnywhere, BlueprintReadOnly, Category=Compression, meta=(DisplayName="Maximum Texture Size", ClampMin = "0.0"), AdvancedDisplay)
+	int32 MaxTextureSize;
+
 	/** When true, the alpha channel of mip-maps and the base image are dithered for smooth LOD transitions. */
 	UPROPERTY(EditAnywhere, BlueprintReadWrite, Category=Texture, AdvancedDisplay)
 	uint32 bDitherMipMapAlpha:1;
@@ -617,6 +621,11 @@ public:
 	 * Marks platform data as transient. This optionally removes persistent or cached data associated with the platform.
 	 */
 	ENGINE_API void MarkPlatformDataTransient();
+
+	/**
+	* Return maximum dimension for this texture type.
+	*/
+	virtual uint32 GetMaximumDimension() const;
 #endif
 
 	/** @return the width of the surface represented by the texture. */
diff --git a/Engine/Source/Runtime/Engine/Classes/Engine/TextureCube.h b/Engine/Source/Runtime/Engine/Classes/Engine/TextureCube.h
--- a/Engine/Source/Runtime/Engine/Classes/Engine/TextureCube.h
+++ b/Engine/Source/Runtime/Engine/Classes/Engine/TextureCube.h
@@ -81,6 +81,13 @@ public:
 	 * @return	Total size of all specified mips, in bytes
 	 */
 	virtual uint32 CalcTextureMemorySizeEnum( ETextureMipCount Enum ) const override;
+
+#if WITH_EDITORONLY_DATA
+	/**
+	* Return maximum dimension for this texture type.
+	*/
+	virtual uint32 GetMaximumDimension() const override;
+#endif
 };
 
 
diff --git a/Engine/Source/Runtime/Engine/Private/Texture.cpp b/Engine/Source/Runtime/Engine/Private/Texture.cpp
--- a/Engine/Source/Runtime/Engine/Private/Texture.cpp
+++ b/Engine/Source/Runtime/Engine/Private/Texture.cpp
@@ -49,6 +49,7 @@ UTexture::UTexture(const class FPostConstructInitializeProperties& PCIP)
 	AdjustHue = 0.0f;
 	AdjustMinAlpha = 0.0f;
 	AdjustMaxAlpha = 1.0f;
+	MaxTextureSize = 0; // means no limitation
 	MipGenSettings = TMGS_FromTextureGroup;
 	SourceArtType_DEPRECATED = TSAT_PNGCompressed;
 	CompositeTextureMode = CTM_NormalRoughnessToAlpha;
@@ -114,6 +115,9 @@ void UTexture::PostEditChangeProperty(FPropertyChangedEvent& PropertyChangedEven
 	{
 		static const FName CompressionSettingsName("CompressionSettings");
 		static const FName DeferCompressionName("DeferCompression");
+#if WITH_EDITORONLY_DATA
+		static const FName MaxTextureSizeName("MaxTextureSize");
+#endif // #if WITH_EDITORONLY_DATA
 
 		const FName PropertyName = PropertyThatChanged->GetFName();
 		if (PropertyName == CompressionSettingsName)
@@ -124,6 +128,19 @@ void UTexture::PostEditChangeProperty(FPropertyChangedEvent& PropertyChangedEven
 		{
 			DeferCompressionWasEnabled = DeferCompression;
 		}
+#if WITH_EDITORONLY_DATA
+		else if (PropertyName == MaxTextureSizeName)
+		{
+			if (MaxTextureSize <= 0)
+			{
+				MaxTextureSize = 0;
+			}
+			else
+			{
+				MaxTextureSize = FMath::Min<int32>(FMath::RoundUpToPowerOfTwo(MaxTextureSize), GetMaximumDimension());
+			}
+		}
+#endif // #if WITH_EDITORONLY_DATA
 
 		bool bPreventSRGB = (CompressionSettings == TC_Alpha || CompressionSettings == TC_Normalmap || CompressionSettings == TC_Masks || CompressionSettings == TC_HDR);		
 		if(bPreventSRGB && SRGB == true)
@@ -911,4 +928,8 @@ void FTextureSource::UseHashAsGuid()
 	}
 }
 
+uint32 UTexture::GetMaximumDimension() const
+{
+	return GetMax2DTextureDimension();
+}
 #endif // #if WITH_EDITORONLY_DATA
diff --git a/Engine/Source/Runtime/Engine/Private/TextureCube.cpp b/Engine/Source/Runtime/Engine/Private/TextureCube.cpp
--- a/Engine/Source/Runtime/Engine/Private/TextureCube.cpp
+++ b/Engine/Source/Runtime/Engine/Private/TextureCube.cpp
@@ -370,3 +370,10 @@ SIZE_T UTextureCube::GetResourceSize(EResourceSizeMode::Type Mode)
 {
 	return CalcTextureMemorySizeEnum(TMC_ResidentMips);
 }
+
+#if WITH_EDITORONLY_DATA
+uint32 UTextureCube::GetMaximumDimension() const
+{
+	return GetMaxCubeTextureDimension();
+}
+#endif // #if WITH_EDITORONLY_DATA
\ No newline at end of file
diff --git a/Engine/Source/Runtime/Engine/Private/TextureDerivedData.cpp b/Engine/Source/Runtime/Engine/Private/TextureDerivedData.cpp
--- a/Engine/Source/Runtime/Engine/Private/TextureDerivedData.cpp
+++ b/Engine/Source/Runtime/Engine/Private/TextureDerivedData.cpp
@@ -141,6 +141,18 @@ static void SerializeForKey(FArchive& Ar, const FTextureBuildSettings& Settings)
 	TempByte = Settings.bApplyKernelToTopMip; Ar << TempByte;
 	TempByte = Settings.CompositeTextureMode; Ar << TempByte;
 	TempFloat = Settings.CompositePower; Ar << TempFloat;
+	// Preserve DDC key for existing content without customized MaxTextureResolution 
+	if (Settings.bLongLatSource)
+	{
+		if (Settings.MaxTextureResolution != 512) // default value for long/lat source
+		{
+			TempUint32 = Settings.MaxTextureResolution; Ar << TempUint32;
+		}
+	}
+	else if (Settings.MaxTextureResolution != TNumericLimits<uint32>::Max())
+	{
+		TempUint32 = Settings.MaxTextureResolution; Ar << TempUint32;
+	}
 }
 
 /**
@@ -288,16 +300,28 @@ static void GetTextureBuildSettings(
 	OutBuildSettings.bComputeBokehAlpha = (Texture.LODGroup == TEXTUREGROUP_Bokeh);
 	OutBuildSettings.bReplicateAlpha = false;
 	OutBuildSettings.bReplicateRed = false;
+	if (Texture.MaxTextureSize > 0)
+	{
+		OutBuildSettings.MaxTextureResolution = Texture.MaxTextureSize;
+	}
 
 	if (Texture.IsA(UTextureCube::StaticClass()))
 	{
 		OutBuildSettings.bCubemap = true;
 		OutBuildSettings.DiffuseConvolveMipLevel = GDiffuseConvolveMipLevel;
+		const UTextureCube* Cube = CastChecked<UTextureCube>(&Texture);
+		OutBuildSettings.bLongLatSource = (Cube->Source.GetNumSlices() == 1);
+		if (OutBuildSettings.bLongLatSource && Texture.MaxTextureSize <= 0)
+		{
+			// long/lat source use 512 as default
+			OutBuildSettings.MaxTextureResolution = 512;
+		}
 	}
 	else
 	{
 		OutBuildSettings.bCubemap = false;
 		OutBuildSettings.DiffuseConvolveMipLevel = 0;
+		OutBuildSettings.bLongLatSource = false;
 	}
 
 	if (Texture.CompressionSettings == TC_Displacementmap)
