Engine: Fixed an incorrect comment for AdjustImageColors

diff --git a/Engine/Source/Developer/TextureCompressor/Private/TextureCompressorModule.cpp b/Engine/Source/Developer/TextureCompressor/Private/TextureCompressorModule.cpp
--- a/Engine/Source/Developer/TextureCompressor/Private/TextureCompressorModule.cpp
+++ b/Engine/Source/Developer/TextureCompressor/Private/TextureCompressorModule.cpp
@@ -1116,7 +1116,7 @@ static void GenerateAngularFilteredMips(TArray<FImage>& InOutMipChain, int32 Num
 ------------------------------------------------------------------------------*/
 
 /**
- * Adjusts the colors of the image using the specified settings (alpha channel will not be modified.)
+ * Adjusts the colors of the image using the specified settings
  *
  * @param	Image		Image to adjust
  * @param	InParams	Color adjustment parameters
