Rendering: Add support for padding textures out to a power of 2 and/or replacing a specified color with transparent black (simple compress-time chroma keying)

diff --git a/Engine/Source/Developer/TextureCompressor/Private/TextureCompressorModule.cpp b/Engine/Source/Developer/TextureCompressor/Private/TextureCompressorModule.cpp
--- a/Engine/Source/Developer/TextureCompressor/Private/TextureCompressorModule.cpp
+++ b/Engine/Source/Developer/TextureCompressor/Private/TextureCompressorModule.cpp
@@ -1118,11 +1118,12 @@ static void GenerateAngularFilteredMips(TArray<FImage>& InOutMipChain, int32 Num
 /**
  * Adjusts the colors of the image using the specified settings
  *
- * @param	Image		Image to adjust
- * @param	InParams	Color adjustment parameters
+ * @param	Image			Image to adjust
+ * @param	InBuildSettings	Image build settings
  */
-static void AdjustImageColors( FImage& Image, const FColorAdjustmentParameters& InParams )
+static void AdjustImageColors( FImage& Image, const FTextureBuildSettings& InBuildSettings )
 {
+	const FColorAdjustmentParameters& InParams = InBuildSettings.ColorAdjustment;
 	check( Image.SizeX > 0 && Image.SizeY > 0 );
 
 	if( !FMath::IsNearlyEqual( InParams.AdjustBrightness, 1.0f, (float)KINDA_SMALL_NUMBER ) ||
@@ -1132,14 +1133,23 @@ static void AdjustImageColors( FImage& Image, const FColorAdjustmentParameters&
 		!FMath::IsNearlyEqual( InParams.AdjustRGBCurve, 1.0f, (float)KINDA_SMALL_NUMBER ) ||
 		!FMath::IsNearlyEqual( InParams.AdjustHue, 0.0f, (float)KINDA_SMALL_NUMBER ) ||
 		!FMath::IsNearlyEqual( InParams.AdjustMinAlpha, 0.0f, (float)KINDA_SMALL_NUMBER ) ||
-		!FMath::IsNearlyEqual( InParams.AdjustMaxAlpha, 1.0f, (float)KINDA_SMALL_NUMBER ) )
+		!FMath::IsNearlyEqual( InParams.AdjustMaxAlpha, 1.0f, (float)KINDA_SMALL_NUMBER ) ||
+		InBuildSettings.bChromaKeyTexture )
 	{
+		const FLinearColor ChromaKeyTarget = InBuildSettings.ChromaKeyColor;
+		const float ChromaKeyThreshold = InBuildSettings.ChromaKeyThreshold + SMALL_NUMBER;
 		const int32 NumPixels = Image.SizeX * Image.SizeY * Image.NumSlices;
 		FLinearColor* ImageColors = Image.AsRGBA32F();
 		const bool bIsSRGB = Image.bSRGB;
 		for( int32 CurPixelIndex = 0; CurPixelIndex < NumPixels; ++CurPixelIndex )
 		{
-			const FLinearColor OriginalColor = ImageColors[ CurPixelIndex ];
+			const FLinearColor OriginalColorRaw = ImageColors[ CurPixelIndex ];
+
+			FLinearColor OriginalColor = OriginalColorRaw;
+			if (InBuildSettings.bChromaKeyTexture && (OriginalColor.Equals(ChromaKeyTarget, ChromaKeyThreshold)))
+			{
+				OriginalColor = FLinearColor::Transparent;
+			}
 
 			// Convert to HSV
 			FLinearColor HSVColor = OriginalColor.LinearRGBToHSV();
@@ -1719,7 +1729,86 @@ private:
 			NumSourceMips = 1;
 		}
 
-		// Check our source can be exported by our current compressor.
+		TArray<FImage> PaddedSourceMips;
+
+		{
+			const FImage& FirstSourceMipImage = InSourceMips[0];
+			int32 TargetTextureSizeX = FirstSourceMipImage.SizeX;
+			int32 TargetTextureSizeY = FirstSourceMipImage.SizeY;
+			bool bPadOrStretchTexture = false;
+
+			const int32 PowerOfTwoTextureSizeX = FMath::RoundUpToPowerOfTwo(TargetTextureSizeX);
+			const int32 PowerOfTwoTextureSizeY = FMath::RoundUpToPowerOfTwo(TargetTextureSizeY);
+			switch (static_cast<const ETexturePowerOfTwoSetting::Type>(BuildSettings.PowerOfTwoMode))
+			{
+			case ETexturePowerOfTwoSetting::None:
+				break;
+
+			case ETexturePowerOfTwoSetting::PadToPowerOfTwo:
+				bPadOrStretchTexture = true;
+				TargetTextureSizeX = PowerOfTwoTextureSizeX;
+				TargetTextureSizeY = PowerOfTwoTextureSizeY;
+				break;
+
+			case ETexturePowerOfTwoSetting::PadToSquarePowerOfTwo:
+				bPadOrStretchTexture = true;
+				TargetTextureSizeX = TargetTextureSizeY = FMath::Max<int32>(PowerOfTwoTextureSizeX, PowerOfTwoTextureSizeY);
+				break;
+
+			default:
+				checkf(false, TEXT("Unknown entry in ETexturePowerOfTwoSetting::Type"));
+				break;
+			}
+
+			if (bPadOrStretchTexture)
+			{
+				// Want to stretch or pad the texture
+				bool bSuitableFormat = FirstSourceMipImage.Format == ERawImageFormat::RGBA32F;
+
+				FImage Temp;
+				if (!bSuitableFormat)
+				{
+					// convert to RGBA32F
+					FirstSourceMipImage.CopyTo(Temp, ERawImageFormat::RGBA32F, false);
+				}
+
+				// space for one source mip and one destination mip
+				const FImage& SourceImage = bSuitableFormat ? FirstSourceMipImage : Temp;
+				FImage& TargetImage = *new (PaddedSourceMips) FImage(TargetTextureSizeX, TargetTextureSizeY, SourceImage.NumSlices, SourceImage.Format);
+				FLinearColor FillColor = BuildSettings.PaddingColor;
+
+				FLinearColor* TargetPtr = (FLinearColor*)TargetImage.RawData.GetData();
+				FLinearColor* SourcePtr = (FLinearColor*)SourceImage.RawData.GetData();
+				check(SourceImage.GetBytesPerPixel() == sizeof(FLinearColor));
+				check(TargetImage.GetBytesPerPixel() == sizeof(FLinearColor));
+
+				const int32 SourceBytesPerLine = SourceImage.SizeX * SourceImage.GetBytesPerPixel();
+				const int32 DestBytesPerLine = TargetImage.SizeX * TargetImage.GetBytesPerPixel();
+				for (int32 SliceIndex = 0; SliceIndex < SourceImage.NumSlices; ++SliceIndex)
+				{
+					for (int32 Y = 0; Y < TargetTextureSizeY; ++Y)
+					{
+						int32 XStart = 0;
+						if (Y < SourceImage.SizeY)
+						{
+							XStart = SourceImage.SizeX;
+							FMemory::Memcpy(TargetPtr, SourcePtr, SourceImage.SizeX * sizeof(FLinearColor));
+							SourcePtr += SourceImage.SizeX;
+							TargetPtr += SourceImage.SizeX;
+						}
+
+						for (int32 XPad = XStart; XPad < TargetImage.SizeX; ++XPad)
+						{
+							*TargetPtr++ = FillColor;
+						}
+					}
+				}
+			}
+		}
+
+		const TArray<FImage>& PostOptionalUpscaleSourceMips = (PaddedSourceMips.Num() > 0) ? PaddedSourceMips : InSourceMips;
+
+		// See if the smallest provided mip image is still too large for the current compressor.
 		int32 LevelsToUsableSource = FMath::Max(0, MaxSourceMipCount - MaxDestMipCount);
 		int32 StartMip = FMath::Max(0, LevelsToUsableSource);
 		bool bBuildSourceImage = StartMip > (NumSourceMips - 1);
@@ -1729,8 +1818,8 @@ private:
 		{			
 			// the source is larger than the compressor allows and no mip image exists to act as a smaller source.
 			// We must generate a suitable source image:
-			bool bSuitableFormat = InSourceMips.Last().Format == ERawImageFormat::RGBA32F;
-			const FImage& BaseImage = InSourceMips.Last();
+			bool bSuitableFormat = PostOptionalUpscaleSourceMips.Last().Format == ERawImageFormat::RGBA32F;
+			const FImage& BaseImage = PostOptionalUpscaleSourceMips.Last();
 
 			if (BaseImage.SizeX != FMath::RoundUpToPowerOfTwo(BaseImage.SizeX) || BaseImage.SizeY != FMath::RoundUpToPowerOfTwo(BaseImage.SizeY))
 			{
@@ -1763,7 +1852,7 @@ private:
 			StartMip--;
 		}
 
-		const TArray<FImage>& SourceMips = bBuildSourceImage ? GeneratedSourceMips : InSourceMips;
+		const TArray<FImage>& SourceMips = bBuildSourceImage ? GeneratedSourceMips : PostOptionalUpscaleSourceMips;
 
 		OutMipChain.Empty(NumOutputMips);
 		// Copy over base mips.
@@ -1812,7 +1901,7 @@ private:
 			}
 
 			// Apply color adjustments
-			AdjustImageColors(*Mip, BuildSettings.ColorAdjustment);
+			AdjustImageColors(*Mip, BuildSettings);
 			if (BuildSettings.bComputeBokehAlpha)
 			{
 				// To get the occlusion in the BokehDOF shader working for all Bokeh textures.
diff --git a/Engine/Source/Developer/TextureCompressor/Public/TextureCompressorModule.h b/Engine/Source/Developer/TextureCompressor/Public/TextureCompressorModule.h
--- a/Engine/Source/Developer/TextureCompressor/Public/TextureCompressorModule.h
+++ b/Engine/Source/Developer/TextureCompressor/Public/TextureCompressorModule.h
@@ -115,6 +115,16 @@ struct FTextureBuildSettings
 	mutable FIntPoint TopMipSize;
 	/** Can the texture be streamed */
 	uint32 bStreamable : 1;
+	/** Whether to chroma key the image, replacing any pixels that match ChromaKeyColor with transparent black */
+	uint32 bChromaKeyTexture : 1;
+	/** How to stretch or pad the texture to a power of 2 size (if necessary); ETexturePowerOfTwoSetting::Type, opaque to avoid dependencies on Engine headers. */
+	uint8 PowerOfTwoMode;
+	/** The color used to pad the texture out if it is resized due to PowerOfTwoMode */
+	FColor PaddingColor;
+	/** The color that will be replaced with transparent black if chroma keying is enabled */
+	FColor ChromaKeyColor;
+	/** The threshold that components have to match for the texel to be considered equal to the ChromaKeyColor when chroma keying (<=, set to 0 to require a perfect exact match) */
+	float ChromaKeyThreshold;
 
 	/** Default settings. */
 	FTextureBuildSettings()
@@ -142,6 +152,11 @@ struct FTextureBuildSettings
 		, LODBias(0)
 		, TopMipSize(0, 0)
 		, bStreamable(false)
+		, bChromaKeyTexture(false)
+		, PowerOfTwoMode(0 /*ETexturePowerOfTwoSetting::None*/)
+		, PaddingColor(FColor::Black)
+		, ChromaKeyColor(FColorList::Magenta)
+		, ChromaKeyThreshold(1.0f / 255.0f)
 	{
 	}
 };
diff --git a/Engine/Source/Runtime/Engine/Classes/Engine/Texture.h b/Engine/Source/Runtime/Engine/Classes/Engine/Texture.h
--- a/Engine/Source/Runtime/Engine/Classes/Engine/Texture.h
+++ b/Engine/Source/Runtime/Engine/Classes/Engine/Texture.h
@@ -60,6 +60,10 @@ enum ECompositeTextureMode
 	// CompositingTexture needs to be a normal map with the same or larger size
 	CTM_NormalRoughnessToAlpha UMETA(DisplayName="Add Normal Roughness To Alpha"),
 	CTM_MAX,
+
+	// Note: These are serialized as as raw values in the texture DDC key, so additional entries
+	// should be added at the bottom; reordering or removing entries will require changing the GUID
+	// in the texture compressor DDC key
 };
 
 UENUM()
@@ -409,6 +413,26 @@ public:
 	UPROPERTY()
 	uint32 bForcePVRTC4:1;
 
+	/** How to pad the texture to a power of 2 size (if necessary) */
+	UPROPERTY(EditAnywhere, BlueprintReadWrite, Category=Texture)
+	TEnumAsByte<enum ETexturePowerOfTwoSetting::Type> PowerOfTwoMode;
+
+	/** The color used to pad the texture out if it is resized due to PowerOfTwoMode */
+	UPROPERTY(EditAnywhere, BlueprintReadWrite, Category=Texture)
+	FColor PaddingColor;
+
+	/** Whether to chroma key the image, replacing any pixels that match ChromaKeyColor with transparent black */
+	UPROPERTY(EditAnywhere, BlueprintReadWrite, Category=Adjustments)
+	bool bChromaKeyTexture;
+
+	/** The threshold that components have to match for the texel to be considered equal to the ChromaKeyColor when chroma keying (<=, set to 0 to require a perfect exact match) */
+	UPROPERTY(EditAnywhere, BlueprintReadWrite, Category=Adjustments, meta=(EditCondition="bChromaKeyTexture", ClampMin="0"))
+	float ChromaKeyThreshold;
+
+	/** The color that will be replaced with transparent black if chroma keying is enabled */
+	UPROPERTY(EditAnywhere, BlueprintReadWrite, Category=Adjustments, meta=(EditCondition="bChromaKeyTexture"))
+	FColor ChromaKeyColor;
+
 	/** Per asset specific setting to define the mip-map generation properties like sharpening and kernel size. */
 	UPROPERTY(EditAnywhere, BlueprintReadWrite, Category=LevelOfDetail)
 	TEnumAsByte<enum TextureMipGenSettings> MipGenSettings;
diff --git a/Engine/Source/Runtime/Engine/Classes/Engine/TextureDefines.h b/Engine/Source/Runtime/Engine/Classes/Engine/TextureDefines.h
--- a/Engine/Source/Runtime/Engine/Classes/Engine/TextureDefines.h
+++ b/Engine/Source/Runtime/Engine/Classes/Engine/TextureDefines.h
@@ -83,4 +83,29 @@ enum TextureMipGenSettings
 	TMGS_Blur4 UMETA(DisplayName="Blur4"),
 	TMGS_Blur5 UMETA(DisplayName="Blur5"),
 	TMGS_MAX,
+
+	// Note: These are serialized as as raw values in the texture DDC key, so additional entries
+	// should be added at the bottom; reordering or removing entries will require changing the GUID
+	// in the texture compressor DDC key
 };
+
+// Options for texture padding mode
+UENUM()
+namespace ETexturePowerOfTwoSetting
+{
+	enum Type
+	{
+		// Do not modify the texture dimensions
+		None,
+
+		// Pad the texture to the nearest power of two size
+		PadToPowerOfTwo,
+
+		// Pad the texture to the nearest square power of two size
+		PadToSquarePowerOfTwo
+
+		// Note: These are serialized as as raw values in the texture DDC key, so additional entries
+		// should be added at the bottom; reordering or removing entries will require changing the GUID
+		// in the texture compressor DDC key
+	};
+}
diff --git a/Engine/Source/Runtime/Engine/Private/Texture.cpp b/Engine/Source/Runtime/Engine/Private/Texture.cpp
--- a/Engine/Source/Runtime/Engine/Private/Texture.cpp
+++ b/Engine/Source/Runtime/Engine/Private/Texture.cpp
@@ -54,6 +54,10 @@ UTexture::UTexture(const FObjectInitializer& ObjectInitializer)
 	MipGenSettings = TMGS_FromTextureGroup;
 	CompositeTextureMode = CTM_NormalRoughnessToAlpha;
 	CompositePower = 1.0f;
+
+	PaddingColor = FColor::Black;
+	ChromaKeyColor = FColorList::Magenta;
+	ChromaKeyThreshold = 1.0f / 255.0f;
 #endif // #if WITH_EDITORONLY_DATA
 
 	if (FApp::CanEverRender() && !IsTemplate())
diff --git a/Engine/Source/Runtime/Engine/Private/Texture2D.cpp b/Engine/Source/Runtime/Engine/Private/Texture2D.cpp
--- a/Engine/Source/Runtime/Engine/Private/Texture2D.cpp
+++ b/Engine/Source/Runtime/Engine/Private/Texture2D.cpp
@@ -230,7 +230,7 @@ void UTexture2D::PostEditUndo()
 void UTexture2D::PostEditChangeProperty(FPropertyChangedEvent& PropertyChangedEvent)
 {
 #if WITH_EDITORONLY_DATA
-	if (!Source.IsPowerOfTwo())
+	if (!Source.IsPowerOfTwo() && (PowerOfTwoMode == ETexturePowerOfTwoSetting::None))
 	{
 		// Force NPT textures to have no mipmaps.
 		MipGenSettings = TMGS_NoMipmaps;
diff --git a/Engine/Source/Runtime/Engine/Private/TextureDerivedData.cpp b/Engine/Source/Runtime/Engine/Private/TextureDerivedData.cpp
--- a/Engine/Source/Runtime/Engine/Private/TextureDerivedData.cpp
+++ b/Engine/Source/Runtime/Engine/Private/TextureDerivedData.cpp
@@ -34,7 +34,8 @@ enum
 // In case of merge conflicts with DDC versions, you MUST generate a new GUID and set this new
 // guid as version
 
-#define TEXTURE_DERIVEDDATA_VER		TEXT("75F87CE100B546AAB8A1C4132AF9DB17")
+#define TEXTURE_DERIVEDDATA_VER		TEXT("814DCC3DC72143F49509781513CB9855")
+
 
 /*------------------------------------------------------------------------------
 	Timing of derived data operations.
@@ -111,6 +112,7 @@ static void SerializeForKey(FArchive& Ar, const FTextureBuildSettings& Settings)
 	uint32 TempUint32;
 	float TempFloat;
 	uint8 TempByte;
+	FColor TempColor;
 
 	TempFloat = Settings.ColorAdjustment.AdjustBrightness; Ar << TempFloat;
 	TempFloat = Settings.ColorAdjustment.AdjustBrightnessCurve; Ar << TempFloat;
@@ -139,18 +141,12 @@ static void SerializeForKey(FArchive& Ar, const FTextureBuildSettings& Settings)
 	TempByte = Settings.bApplyKernelToTopMip; Ar << TempByte;
 	TempByte = Settings.CompositeTextureMode; Ar << TempByte;
 	TempFloat = Settings.CompositePower; Ar << TempFloat;
-	// Preserve DDC key for existing content without customized MaxTextureResolution 
-	if (Settings.bLongLatSource)
-	{
-		if (Settings.MaxTextureResolution != 512) // default value for long/lat source
-		{
-			TempUint32 = Settings.MaxTextureResolution; Ar << TempUint32;
-		}
-	}
-	else if (Settings.MaxTextureResolution != TNumericLimits<uint32>::Max())
-	{
-		TempUint32 = Settings.MaxTextureResolution; Ar << TempUint32;
-	}
+	TempUint32 = Settings.MaxTextureResolution; Ar << TempUint32;
+	TempByte = Settings.PowerOfTwoMode; Ar << TempByte;
+	TempColor = Settings.PaddingColor; Ar << TempColor;
+	TempByte = Settings.bChromaKeyTexture; Ar << TempByte;
+	TempColor = Settings.ChromaKeyColor; Ar << TempColor;
+	TempFloat = Settings.ChromaKeyThreshold; Ar << TempFloat;
 }
 
 /**
@@ -353,6 +349,11 @@ static void GetTextureBuildSettings(
 	OutBuildSettings.CompositePower = Texture.CompositePower;
 	OutBuildSettings.LODBias = GSystemSettings.TextureLODSettings.CalculateLODBias(Texture.Source.GetSizeX(), Texture.Source.GetSizeY(), Texture.LODGroup, Texture.LODBias, Texture.NumCinematicMipLevels, Texture.MipGenSettings);
 	OutBuildSettings.bStreamable = !Texture.NeverStream && (Texture.LODGroup != TEXTUREGROUP_UI) && (Cast<const UTexture2D>(&Texture) != NULL);
+	OutBuildSettings.PowerOfTwoMode = Texture.PowerOfTwoMode;
+	OutBuildSettings.PaddingColor = Texture.PaddingColor;
+	OutBuildSettings.ChromaKeyColor = Texture.ChromaKeyColor;
+	OutBuildSettings.bChromaKeyTexture = Texture.bChromaKeyTexture;
+	OutBuildSettings.ChromaKeyThreshold = Texture.ChromaKeyThreshold;
 }
 
 /**
