Fix for UE-6886 - Editor crashing when changing Max Texture Size of a texture to a power of 2 below 64

diff --git a/Engine/Source/Developer/TextureCompressor/Private/TextureCompressorModule.cpp b/Engine/Source/Developer/TextureCompressor/Private/TextureCompressorModule.cpp
--- a/Engine/Source/Developer/TextureCompressor/Private/TextureCompressorModule.cpp
+++ b/Engine/Source/Developer/TextureCompressor/Private/TextureCompressorModule.cpp
@@ -1856,9 +1856,10 @@ private:
 
 		OutMipChain.Empty(NumOutputMips);
 		// Copy over base mips.
-		int32 CopyCount = FMath::Min(NumOutputMips - StartMip, SourceMips.Num());
-		check(StartMip <= SourceMips.Num());
-		for (int32 MipIndex = StartMip; MipIndex < CopyCount; ++MipIndex)
+		check(StartMip < SourceMips.Num());
+		int32 CopyCount = SourceMips.Num() - StartMip;
+
+		for (int32 MipIndex = StartMip; MipIndex < StartMip + CopyCount; ++MipIndex)
 		{
 			const FImage& Image = SourceMips[MipIndex];
 			const int32 SrcWidth = Image.SizeX;
