UE4 - Import cubemap textures on different threads when doing mips >= 128 texels (we're seeing > 20minutes on 2k textures). This is a workaround, the algorithm itself should be rewritten to not be N^2!

diff --git a/Engine/Source/Developer/TextureCompressor/Private/TextureCompressorModule.cpp b/Engine/Source/Developer/TextureCompressor/Private/TextureCompressorModule.cpp
--- a/Engine/Source/Developer/TextureCompressor/Private/TextureCompressorModule.cpp
+++ b/Engine/Source/Developer/TextureCompressor/Private/TextureCompressorModule.cpp
@@ -967,12 +967,12 @@ static FLinearColor IntegrateAngularArea(FImage& Image, FVector FilterDirectionW
 }
 
 // @return 2 * computed triangle area 
-static float TriangleArea2_3D(FVector A, FVector B, FVector C)
+static inline float TriangleArea2_3D(FVector A, FVector B, FVector C)
 {
 	return ((A-B) ^ (C-B)).Size();
 }
 
-static float ComputeTexelArea(uint32 x, uint32 y, float InvSideExtentMul2)
+static inline float ComputeTexelArea(uint32 x, uint32 y, float InvSideExtentMul2)
 {
 	float fU = x * InvSideExtentMul2 - 1;
 	float fV = y * InvSideExtentMul2 - 1;
@@ -1013,15 +1013,75 @@ static void GenerateAngularFilteredMip(FImage* DestMip, FImage& SrcMip, float Co
 		}
 	}
 
-	for(int32 Face = 0; Face < 6; ++Face)
+	// We start getting gains running threaded upwards of sizes >= 128
+	if (SrcMip.SizeX >= 128)
 	{
-		FImageView2D DestMipView(*DestMip, Face);
-		for(int32 y = 0; y < Extent; ++y)
+		// Quick workaround: Do a thread per mip
+		struct FAsyncGenerateMipsPerFaceWorker : public FNonAbandonableTask
 		{
-			for(int32 x = 0; x < Extent; ++x)
+			int32 Face;
+			FImage* DestMip;
+			int32 Extent;
+			float ConeAngle;
+			const float* TexelAreaArray;
+			FImage* SrcMip;
+			FAsyncGenerateMipsPerFaceWorker(int32 InFace, FImage* InDestMip, int32 InExtent, float InConeAngle, const float* InTexelAreaArray, FImage* InSrcMip) :
+				Face(InFace),
+				DestMip(InDestMip),
+				Extent(InExtent),
+				ConeAngle(InConeAngle),
+				TexelAreaArray(InTexelAreaArray),
+				SrcMip(InSrcMip)
+			{
+			}
+
+			void DoWork()
+			{
+				const float InvSideExtent = 1.0f / Extent;
+				FImageView2D DestMipView(*DestMip, Face);
+				for (int32 y = 0; y < Extent; ++y)
+				{
+					for (int32 x = 0; x < Extent; ++x)
+					{
+						FVector DirectionWS = ComputeWSCubeDirectionAtTexelCenter(Face, x, y, InvSideExtent);
+						DestMipView.Access(x, y) = IntegrateAngularArea(*SrcMip, DirectionWS, ConeAngle, TexelAreaArray);
+					}
+				}
+			}
+
+			FORCEINLINE TStatId GetStatId() const
 			{
-				FVector DirectionWS = ComputeWSCubeDirectionAtTexelCenter(Face, x, y, InvSideExtent);
-				DestMipView.Access(x,y) = IntegrateAngularArea(SrcMip, DirectionWS, ConeAngle, TexelAreaArray.GetData());
+				RETURN_QUICK_DECLARE_CYCLE_STAT(FAsyncGenerateMipsPerFaceWorker, STATGROUP_ThreadPoolAsyncTasks);
+			}
+		};
+
+		typedef FAsyncTask<FAsyncGenerateMipsPerFaceWorker> FAsyncGenerateMipsPerFaceTask;
+		TIndirectArray<FAsyncGenerateMipsPerFaceTask> AsyncTasks;
+
+		for (int32 Face = 0; Face < 6; ++Face)
+		{
+			auto* AsyncTask = new(AsyncTasks) FAsyncGenerateMipsPerFaceTask(Face, DestMip, Extent, ConeAngle, TexelAreaArray.GetData(), &SrcMip);
+			AsyncTask->StartBackgroundTask();
+		}
+
+		for (int32 TaskIndex = 0; TaskIndex < AsyncTasks.Num(); ++TaskIndex)
+		{
+			auto& AsyncTask = AsyncTasks[TaskIndex];
+			AsyncTask.EnsureCompletion();
+		}
+	}
+	else
+	{
+		for (int32 Face = 0; Face < 6; ++Face)
+		{
+			FImageView2D DestMipView(*DestMip, Face);
+			for (int32 y = 0; y < Extent; ++y)
+			{
+				for (int32 x = 0; x < Extent; ++x)
+				{
+					FVector DirectionWS = ComputeWSCubeDirectionAtTexelCenter(Face, x, y, InvSideExtent);
+					DestMipView.Access(x, y) = IntegrateAngularArea(SrcMip, DirectionWS, ConeAngle, TexelAreaArray.GetData());
+				}
 			}
 		}
 	}
