fixed UE-19178 Mip Mapping: Replace simple average with high quality alternative.
 alpha channel was missing mip sharpening for color textures (can cause a lot of texture to be recooked)

diff --git a/Engine/Source/Developer/TextureCompressor/Private/TextureCompressorModule.cpp b/Engine/Source/Developer/TextureCompressor/Private/TextureCompressorModule.cpp
--- a/Engine/Source/Developer/TextureCompressor/Private/TextureCompressorModule.cpp
+++ b/Engine/Source/Developer/TextureCompressor/Private/TextureCompressorModule.cpp
@@ -327,7 +327,7 @@ static void GenerateSharpenedMipB8G8R8A8Templ(
 
 			if ( bSharpenWithoutColorShift )
 			{
-				float NewLuminance = 0;
+				FLinearColor SharpenedColor(0, 0, 0, 0);
 
 				for ( uint32 KernelY = 0; KernelY < Kernel.GetFilterTableSize();  ++KernelY )
 				{
@@ -335,12 +335,12 @@ static void GenerateSharpenedMipB8G8R8A8Templ(
 					{
 						float Weight = Kernel.GetAt( KernelX, KernelY );
 						FLinearColor Sample = LookupSourceMip<AddressMode>( SourceImageData, SourceX + KernelX - KernelCenter, SourceY + KernelY - KernelCenter );
-						float LuminanceSample = Sample.ComputeLuminance();
-
-						NewLuminance += Weight * LuminanceSample;
+						SharpenedColor += Weight * Sample;
 					}
 				}
 
+				float NewLuminance = SharpenedColor.ComputeLuminance();
+
 				// simple 2x2 kernel to compute the color
 				FilteredColor =
 					( LookupSourceMip<AddressMode>( SourceImageData, SourceX + 0, SourceY + 0 )
@@ -357,6 +357,9 @@ static void GenerateSharpenedMipB8G8R8A8Templ(
 					FilteredColor.G *= Factor;
 					FilteredColor.B *= Factor;
 				}
+
+				// We also want to sharpen the alpha channel (was missing before)
+				FilteredColor.A = SharpenedColor.A;
 			}
 			else
 			{
diff --git a/Engine/Source/Runtime/Engine/Private/TextureDerivedData.cpp b/Engine/Source/Runtime/Engine/Private/TextureDerivedData.cpp
--- a/Engine/Source/Runtime/Engine/Private/TextureDerivedData.cpp
+++ b/Engine/Source/Runtime/Engine/Private/TextureDerivedData.cpp
@@ -142,7 +142,20 @@ static void SerializeForKey(FArchive& Ar, const FTextureBuildSettings& Settings)
 	TempByte = Settings.bReplicateRed; Ar << TempByte;
 	TempByte = Settings.bReplicateAlpha; Ar << TempByte;
 	TempByte = Settings.bDownsampleWithAverage; Ar << TempByte;
-	TempByte = Settings.bSharpenWithoutColorShift; Ar << TempByte;
+	
+	{
+		TempByte = Settings.bSharpenWithoutColorShift;
+
+		if(Settings.bSharpenWithoutColorShift && Settings.MipSharpening != 0.0f)
+		{
+			// bSharpenWithoutColorShift prevented alpha sharpening. This got fixed
+			// Here we update the key to get those cases recooked.
+			TempByte = 2;
+		}
+
+		Ar << TempByte;
+	}
+
 	TempByte = Settings.bBorderColorBlack; Ar << TempByte;
 	TempByte = Settings.bFlipGreenChannel; Ar << TempByte;
 	TempByte = Settings.bApplyKernelToTopMip; Ar << TempByte;
