Fix resolution of long package names when using unique editor builds. In these situations the base directory is set to the project binaries directory, so ContentPathShort resolves to the project content directory.

Also fix an assumed DLL path for NVTT.

diff --git a/Engine/Source/Developer/TextureCompressor/Private/TextureCompressorModule.cpp b/Engine/Source/Developer/TextureCompressor/Private/TextureCompressorModule.cpp
--- a/Engine/Source/Developer/TextureCompressor/Private/TextureCompressorModule.cpp
+++ b/Engine/Source/Developer/TextureCompressor/Private/TextureCompressorModule.cpp
@@ -10,6 +10,7 @@
 #include "Engine/Texture.h"
 #include "Interfaces/ITargetPlatformManagerModule.h"
 #include "Interfaces/ITextureFormat.h"
+#include "Misc/Paths.h"
 #include "ImageCore.h"
 
 #if PLATFORM_WINDOWS
@@ -1933,9 +1934,9 @@ public:
 	{
 #if PLATFORM_WINDOWS
 	#if PLATFORM_64BITS
-		nvTextureToolsHandle = FPlatformProcess::GetDllHandle(TEXT("../../../Engine/Binaries/ThirdParty/nvTextureTools/Win64/nvtt_64.dll"));
+		nvTextureToolsHandle = FPlatformProcess::GetDllHandle(*(FPaths::EngineDir() / TEXT("Binaries/ThirdParty/nvTextureTools/Win64/nvtt_64.dll")));
 	#else	//32-bit platform
-		nvTextureToolsHandle = FPlatformProcess::GetDllHandle(TEXT("../../../Engine/Binaries/ThirdParty/nvTextureTools/Win32/nvtt_.dll"));
+		nvTextureToolsHandle = FPlatformProcess::GetDllHandle(*(FPaths::EngineDir() / TEXT("Binaries/ThirdParty/nvTextureTools/Win32/nvtt_.dll")));
 	#endif
 #endif	//PLATFORM_WINDOWS
 	}
diff --git a/Engine/Source/Runtime/CoreUObject/Private/Misc/PackageName.cpp b/Engine/Source/Runtime/CoreUObject/Private/Misc/PackageName.cpp
--- a/Engine/Source/Runtime/CoreUObject/Private/Misc/PackageName.cpp
+++ b/Engine/Source/Runtime/CoreUObject/Private/Misc/PackageName.cpp
@@ -259,11 +259,14 @@ private:
 		
 		ContentPathToRoot.Empty(11);
 		ContentPathToRoot.Emplace(EngineRootPath, EngineContentPath);
-#if IS_MONOLITHIC
-		ContentPathToRoot.Emplace(GameRootPath,   ContentPathShort);
-#else
-		ContentPathToRoot.Emplace(EngineRootPath, ContentPathShort);
-#endif
+		if (FPaths::IsSamePath(GameContentPath, ContentPathShort))
+		{
+			ContentPathToRoot.Emplace(GameRootPath, ContentPathShort);
+		}
+		else
+		{
+			ContentPathToRoot.Emplace(EngineRootPath, ContentPathShort);
+		}
 		ContentPathToRoot.Emplace(EngineRootPath, EngineShadersPath);
 		ContentPathToRoot.Emplace(EngineRootPath, EngineShadersPathShort);
 		ContentPathToRoot.Emplace(GameRootPath,   GameContentPath);
