Fixed Texture Editor HSV controls to work with HDR textures.

diff --git a/Engine/Source/Developer/TextureCompressor/Private/TextureCompressorModule.cpp b/Engine/Source/Developer/TextureCompressor/Private/TextureCompressorModule.cpp
--- a/Engine/Source/Developer/TextureCompressor/Private/TextureCompressorModule.cpp
+++ b/Engine/Source/Developer/TextureCompressor/Private/TextureCompressorModule.cpp
@@ -1402,79 +1402,111 @@ static void AdjustImageColors( FImage& Image, const FTextureBuildSettings& InBui
 		const int32 NumPixels = Image.SizeX * Image.SizeY * Image.NumSlices;
 		FLinearColor* ImageColors = Image.AsRGBA32F();
 
-		for( int32 CurPixelIndex = 0; CurPixelIndex < NumPixels; ++CurPixelIndex )
+		int32 NumJobs = FTaskGraphInterface::Get().GetNumWorkerThreads();
+		int32 NumPixelsEachJob = NumPixels / NumJobs;
+		if (NumPixelsEachJob * NumJobs < NumPixels)
 		{
-			const FLinearColor OriginalColorRaw = ImageColors[ CurPixelIndex ];
+			++NumPixelsEachJob;
+		}
+
+		const static int MinPixelsToMultithread = 16384;
+		bool bForceSingleThread = (NumPixels < MinPixelsToMultithread) || GIsEditorLoadingPackage || GIsCookerLoadingPackage || IsInAsyncLoadingThread();
 
-			FLinearColor OriginalColor = OriginalColorRaw;
-			if (InBuildSettings.bChromaKeyTexture && (OriginalColor.Equals(ChromaKeyTarget, ChromaKeyThreshold)))
+		TFunctionRef<void (int32)> AdjustImageColorsFunc = [&](int32 Index)
+		{
+			int32 StartIndex = Index * NumPixelsEachJob;
+			int32 EndIndex = FMath::Min(StartIndex + NumPixelsEachJob, NumPixels);
+			for (int32 CurPixelIndex = StartIndex; CurPixelIndex < EndIndex; ++CurPixelIndex)
 			{
-				OriginalColor = FLinearColor::Transparent;
-			}
+				const FLinearColor OriginalColorRaw = ImageColors[CurPixelIndex];
 
-			// Convert to HSV
-			FLinearColor HSVColor = OriginalColor.LinearRGBToHSV();
-			float& PixelHue = HSVColor.R;
-			float& PixelSaturation = HSVColor.G;
-			float& PixelValue = HSVColor.B;
+				FLinearColor OriginalColor = OriginalColorRaw;
+				if (InBuildSettings.bChromaKeyTexture && (OriginalColor.Equals(ChromaKeyTarget, ChromaKeyThreshold)))
+				{
+					OriginalColor = FLinearColor::Transparent;
+				}
 
-			// Apply brightness adjustment
-			PixelValue *= InParams.AdjustBrightness;
+				// Convert to HSV
+				FLinearColor HSVColor = OriginalColor.LinearRGBToHSV();
+				float& PixelHue = HSVColor.R;
+				float& PixelSaturation = HSVColor.G;
+				float& PixelValue = HSVColor.B;
 
-			// Apply brightness power adjustment
-			if( !FMath::IsNearlyEqual( InParams.AdjustBrightnessCurve, 1.0f, (float)KINDA_SMALL_NUMBER ) && InParams.AdjustBrightnessCurve != 0.0f )
-			{
-				// Raise HSV.V to the specified power
-				PixelValue = FMath::Pow( PixelValue, InParams.AdjustBrightnessCurve );
-			}
+				float OriginalLuminance = PixelValue;
 
-			// Apply "vibrance" adjustment
-			if( !FMath::IsNearlyZero( InParams.AdjustVibrance, (float)KINDA_SMALL_NUMBER ) )
-			{
-				const float SatRaisePow = 5.0f;
-				const float InvSatRaised = FMath::Pow( 1.0f - PixelSaturation, SatRaisePow );
+				// Apply brightness adjustment
+				PixelValue *= InParams.AdjustBrightness;
+
+				// Apply brightness power adjustment
+				if (!FMath::IsNearlyEqual(InParams.AdjustBrightnessCurve, 1.0f, (float)KINDA_SMALL_NUMBER) && InParams.AdjustBrightnessCurve != 0.0f)
+				{
+					// Raise HSV.V to the specified power
+					PixelValue = FMath::Pow(PixelValue, InParams.AdjustBrightnessCurve);
+				}
 
-				const float ClampedVibrance = FMath::Clamp( InParams.AdjustVibrance, 0.0f, 1.0f );
-				const float HalfVibrance = ClampedVibrance * 0.5f;
+				// Apply "vibrance" adjustment
+				if (!FMath::IsNearlyZero(InParams.AdjustVibrance, (float)KINDA_SMALL_NUMBER))
+				{
+					const float SatRaisePow = 5.0f;
+					const float InvSatRaised = FMath::Pow(1.0f - PixelSaturation, SatRaisePow);
 
-				const float SatProduct = HalfVibrance * InvSatRaised;
+					const float ClampedVibrance = FMath::Clamp(InParams.AdjustVibrance, 0.0f, 1.0f);
+					const float HalfVibrance = ClampedVibrance * 0.5f;
 
-				PixelSaturation += SatProduct;
-			}
+					const float SatProduct = HalfVibrance * InvSatRaised;
 
-			// Apply saturation adjustment
-			PixelSaturation *= InParams.AdjustSaturation;
+					PixelSaturation += SatProduct;
+				}
 
-			// Apply hue adjustment
-			PixelHue += InParams.AdjustHue;
+				// Apply saturation adjustment
+				PixelSaturation *= InParams.AdjustSaturation;
 
-			// Clamp HSV values
-			{
-				PixelHue = FMath::Fmod( PixelHue, 360.0f );
-				if( PixelHue < 0.0f )
+				// Apply hue adjustment
+				PixelHue += InParams.AdjustHue;
+
+				// Clamp HSV values
 				{
-					// Keep the hue value positive as HSVToLinearRGB prefers that
-					PixelHue += 360.0f;
+					PixelHue = FMath::Fmod(PixelHue, 360.0f);
+					if (PixelHue < 0.0f)
+					{
+						// Keep the hue value positive as HSVToLinearRGB prefers that
+						PixelHue += 360.0f;
+					}
+					PixelSaturation = FMath::Clamp(PixelSaturation, 0.0f, 1.0f);
+
+					// Clamp brightness if non-HDR
+					if (!InBuildSettings.bHDRSource)
+					{
+						PixelValue = FMath::Clamp(PixelValue, 0.0f, 1.0f);
+					}
 				}
-				PixelSaturation = FMath::Clamp( PixelSaturation, 0.0f, 1.0f );
-				PixelValue = FMath::Clamp( PixelValue, 0.0f, 1.0f );
-			}
 
-			// Convert back to a linear color
-			FLinearColor LinearColor = HSVColor.HSVToLinearRGB();
+				// Convert back to a linear color
+				FLinearColor LinearColor = HSVColor.HSVToLinearRGB();
 
-			// Apply RGB curve adjustment (linear space)
-			if( !FMath::IsNearlyEqual( InParams.AdjustRGBCurve, 1.0f, (float)KINDA_SMALL_NUMBER ) && InParams.AdjustRGBCurve != 0.0f )
-			{
-				LinearColor.R = FMath::Pow( LinearColor.R, InParams.AdjustRGBCurve );
-				LinearColor.G = FMath::Pow( LinearColor.G, InParams.AdjustRGBCurve );
-				LinearColor.B = FMath::Pow( LinearColor.B, InParams.AdjustRGBCurve );
+				// Apply RGB curve adjustment (linear space)
+				if (!FMath::IsNearlyEqual(InParams.AdjustRGBCurve, 1.0f, (float)KINDA_SMALL_NUMBER) && InParams.AdjustRGBCurve != 0.0f)
+				{
+					LinearColor.R = FMath::Pow(LinearColor.R, InParams.AdjustRGBCurve);
+					LinearColor.G = FMath::Pow(LinearColor.G, InParams.AdjustRGBCurve);
+					LinearColor.B = FMath::Pow(LinearColor.B, InParams.AdjustRGBCurve);
+				}
+
+				// Clamp HDR RGB channels to 1 or the original luminance (max original RGB channel value), whichever is greater
+				if (InBuildSettings.bHDRSource)
+				{
+					LinearColor.R = FMath::Clamp(LinearColor.R, 0.0f, (OriginalLuminance > 1.0f ? OriginalLuminance : 1.0f));
+					LinearColor.G = FMath::Clamp(LinearColor.G, 0.0f, (OriginalLuminance > 1.0f ? OriginalLuminance : 1.0f));
+					LinearColor.B = FMath::Clamp(LinearColor.B, 0.0f, (OriginalLuminance > 1.0f ? OriginalLuminance : 1.0f));
+				}
+
+				// Remap the alpha channel
+				LinearColor.A = FMath::Lerp(InParams.AdjustMinAlpha, InParams.AdjustMaxAlpha, OriginalColor.A);
+				ImageColors[CurPixelIndex] = LinearColor;
 			}
+		};
 
-			// Remap the alpha channel
-			LinearColor.A = FMath::Lerp(InParams.AdjustMinAlpha, InParams.AdjustMaxAlpha, OriginalColor.A);
-			ImageColors[ CurPixelIndex ] = LinearColor;
-		}
+		ParallelFor(NumJobs, AdjustImageColorsFunc, bForceSingleThread);
 	}
 }
 
diff --git a/Engine/Source/Developer/TextureCompressor/Public/TextureCompressorModule.h b/Engine/Source/Developer/TextureCompressor/Public/TextureCompressorModule.h
--- a/Engine/Source/Developer/TextureCompressor/Public/TextureCompressorModule.h
+++ b/Engine/Source/Developer/TextureCompressor/Public/TextureCompressorModule.h
@@ -81,6 +81,8 @@ struct FTextureBuildSettings
 	uint32 MaxTextureResolution;
 	/** Format of the compressed texture, used to choose a compression DLL. */
 	FName TextureFormatName;
+	/** Whether the texture being built contains HDR source data */
+	bool bHDRSource;
 	/** Mipmap generation settings. */
 	uint8 MipGenSettings; // TextureMipGenSettings, opaque to avoid dependencies on engine headers.
 	/** Whether the texture being built is a cubemap. */
@@ -179,6 +181,7 @@ struct FTextureBuildSettings
 		, ChromaKeyColor(FColorList::Magenta)
 		, ChromaKeyThreshold(1.0f / 255.0f)
 		, CompressionQuality(-1)
+		, bHDRSource(false)
 	{
 	}
 
diff --git a/Engine/Source/Runtime/Engine/Classes/Engine/Texture.h b/Engine/Source/Runtime/Engine/Classes/Engine/Texture.h
--- a/Engine/Source/Runtime/Engine/Classes/Engine/Texture.h
+++ b/Engine/Source/Runtime/Engine/Classes/Engine/Texture.h
@@ -745,6 +745,7 @@ public:
 	//~ Begin UObject Interface.
 #if WITH_EDITOR
 	ENGINE_API virtual void PostEditChangeProperty(FPropertyChangedEvent& PropertyChangedEvent) override;
+	ENGINE_API virtual bool CanEditChange(const UProperty* InProperty) const override;
 #endif // WITH_EDITOR
 	ENGINE_API virtual void Serialize(FArchive& Ar) override;
 	ENGINE_API virtual void PostInitProperties() override;
diff --git a/Engine/Source/Runtime/Engine/Private/Texture.cpp b/Engine/Source/Runtime/Engine/Private/Texture.cpp
--- a/Engine/Source/Runtime/Engine/Private/Texture.cpp
+++ b/Engine/Source/Runtime/Engine/Private/Texture.cpp
@@ -124,6 +124,21 @@ int32 UTexture::GetCachedLODBias() const
 }
 
 #if WITH_EDITOR
+bool UTexture::CanEditChange(const UProperty* InProperty) const
+{
+	if (InProperty)
+	{
+		FString PropertyName = InProperty->GetName();
+
+		if (PropertyName == GET_MEMBER_NAME_STRING_CHECKED(UTexture, AdjustVibrance))
+		{
+			return !HasHDRSource();
+		}
+	}
+
+	return true;
+}
+
 void UTexture::PostEditChangeProperty(FPropertyChangedEvent& PropertyChangedEvent)
 {
 	Super::PostEditChangeProperty(PropertyChangedEvent);
