Added comment re: when AdjustImageColors should be multithreaded.

diff --git a/Engine/Source/Developer/TextureCompressor/Private/TextureCompressorModule.cpp b/Engine/Source/Developer/TextureCompressor/Private/TextureCompressorModule.cpp
--- a/Engine/Source/Developer/TextureCompressor/Private/TextureCompressorModule.cpp
+++ b/Engine/Source/Developer/TextureCompressor/Private/TextureCompressorModule.cpp
@@ -1409,6 +1409,10 @@ static void AdjustImageColors( FImage& Image, const FTextureBuildSettings& InBui
 			++NumPixelsEachJob;
 		}
 
+		// bForceSingleThread is set to true when: 
+		// (a) editor or cooker is loading as this is when the derived data cache is rebuilt as it will already be limited to a single thread 
+		//     and thus overhead of multithreading will simply make it slower
+		// (b) texture is smaller than 256x256 as it will cost more to multithread than to run on single thread.	
 		const static int MinPixelsToMultithread = 16384;
 		bool bForceSingleThread = (NumPixels < MinPixelsToMultithread) || GIsEditorLoadingPackage || GIsCookerLoadingPackage || IsInAsyncLoadingThread();
 
