Fix crash when taking a ref of a lambda function (Thanks Christina Tempelaarl!)

diff --git a/Engine/Source/Developer/TextureCompressor/Private/TextureCompressorModule.cpp b/Engine/Source/Developer/TextureCompressor/Private/TextureCompressorModule.cpp
--- a/Engine/Source/Developer/TextureCompressor/Private/TextureCompressorModule.cpp
+++ b/Engine/Source/Developer/TextureCompressor/Private/TextureCompressorModule.cpp
@@ -1417,7 +1417,7 @@ static void AdjustImageColors( FImage& Image, const FTextureBuildSettings& InBui
 		const static int MinPixelsToMultithread = 16384;
 		bool bForceSingleThread = (NumPixels < MinPixelsToMultithread) || GIsEditorLoadingPackage || GIsCookerLoadingPackage || IsInAsyncLoadingThread();
 
-		TFunctionRef<void (int32)> AdjustImageColorsFunc = [&](int32 Index)
+		TFunction<void (int32)> AdjustImageColorsFunc = [&](int32 Index)
 		{
 			int32 StartIndex = Index * NumPixelsEachJob;
 			int32 EndIndex = FMath::Min(StartIndex + NumPixelsEachJob, NumPixels);
