Add YCoCg compression with per block scaling to the texture compression module
Runtime virtual texture uses streaming virtual texture to fill pages at distance
We recently moved to YCoCg encoding in the compressed base color texture to address visual quality issues
The runtime virtual texture used a per block scaling, but the offline generated streaming texture didn't
The offline texture looks much worse without the block scaling so that we see color artifacts in the distance before this change

diff --git a/Engine/Source/Developer/TextureCompressor/Private/TextureCompressorModule.cpp b/Engine/Source/Developer/TextureCompressor/Private/TextureCompressorModule.cpp
--- a/Engine/Source/Developer/TextureCompressor/Private/TextureCompressorModule.cpp
+++ b/Engine/Source/Developer/TextureCompressor/Private/TextureCompressorModule.cpp
@@ -1625,6 +1625,63 @@ static bool DetectAlphaChannel(const FImage& InImage)
 	return false;
 }
 
+/** Calculate a scale per 4x4 block of each image, and apply it to the red/green channels. Store scale in the blue channel. */
+static void ApplyYCoCgBlockScale(TArray<FImage>& InOutMipChain)
+{
+	const uint32 MipCount = InOutMipChain.Num();
+	for (uint32 MipIndex = 0; MipIndex < MipCount; ++MipIndex)
+	{
+		FImage& SrcMip = InOutMipChain[MipIndex];
+		FLinearColor* FirstColor = SrcMip.AsRGBA32F();
+
+		int32 BlockWidthX = SrcMip.SizeX / 4;
+		int32 BlockWidthY = SrcMip.SizeY / 4;
+
+		for (int32 Slice = 0; Slice < SrcMip.NumSlices; ++Slice)
+		{
+			FLinearColor* SliceFirstColor = FirstColor + (SrcMip.SizeX * SrcMip.SizeY * Slice);
+
+			for (int32 Y = 0; Y < BlockWidthY; ++Y)
+			{
+				FLinearColor* RowFirstColor = SliceFirstColor + (Y * 4 * SrcMip.SizeY);
+
+				for (int32 X = 0; X < BlockWidthX; ++X)
+				{
+					FLinearColor* BlockFirstColor = RowFirstColor + (X * 4);
+
+					// Iterate block to find MaxComponent
+					float MaxComponent = 0.f;
+					for (int32 BlockY = 0; BlockY < 4; ++BlockY)
+					{
+						FLinearColor* Color = BlockFirstColor + (BlockY * SrcMip.SizeY);
+						for (int32 BlockX = 0; BlockX < 4; ++BlockX, ++Color)
+						{
+							MaxComponent = FMath::Max(FMath::Abs(Color->R - 128.f / 255.f), MaxComponent);
+							MaxComponent = FMath::Max(FMath::Abs(Color->G - 128.f / 255.f), MaxComponent);
+						}
+					}
+
+					const float Scale = (MaxComponent < 32.f / 255.f) ? 4.f : (MaxComponent < 64.f / 255.f) ? 2.f : 1.f;
+					const float OutB = (Scale - 1.f) * 8.f / 255.f;
+
+					// Iterate block to modify for scale
+					for (int32 BlockY = 0; BlockY < 4; ++BlockY)
+					{
+						FLinearColor* Color = BlockFirstColor + (BlockY * SrcMip.SizeY);
+						for (int32 BlockX = 0; BlockX < 4; ++BlockX, ++Color)
+						{
+							const float OutR = (Color->R - 128.f / 255.f) * Scale + 128.f / 255.f;
+							const float OutG = (Color->G - 128.f / 255.f) * Scale + 128.f / 255.f;
+
+							*Color = FLinearColor(OutR, OutG, OutB, Color->A);
+						}
+					}
+				}
+			}
+		}
+	}
+}
+
 float RoughnessToSpecularPower(float Roughness)
 {
 	float Div = FMath::Pow(Roughness, 4);
@@ -2243,6 +2300,10 @@ private:
 		{
 			ReplicateAlphaChannel(OutMipChain);
 		}
+		if (BuildSettings.bApplyYCoCgBlockScale)
+		{
+			ApplyYCoCgBlockScale(OutMipChain);
+		}
 
 		return true;
 	}
diff --git a/Engine/Source/Developer/TextureCompressor/Public/TextureCompressorModule.h b/Engine/Source/Developer/TextureCompressor/Public/TextureCompressorModule.h
--- a/Engine/Source/Developer/TextureCompressor/Public/TextureCompressorModule.h
+++ b/Engine/Source/Developer/TextureCompressor/Public/TextureCompressorModule.h
@@ -115,6 +115,8 @@ struct FTextureBuildSettings
 	uint32 bBorderColorBlack : 1;
 	/** Whether the green channel should be flipped. Typically only done on normal maps. */
 	uint32 bFlipGreenChannel : 1;
+	/** Calculate and apply a scale for YCoCg textures. This calculates a scale for each 4x4 block, applies it to the red and green channels and stores the scale in the blue channel. */
+	uint32 bApplyYCoCgBlockScale : 1;
 	/** 1:apply mip sharpening/blurring kernel to top mip as well (at half the kernel size), 0:don't */
 	uint32 bApplyKernelToTopMip : 1;
 	/** 1: renormalizes the top mip (only useful for normal maps, prevents artists errors and adds quality) 0:don't */
@@ -187,6 +189,7 @@ struct FTextureBuildSettings
 		, bSharpenWithoutColorShift(false)
 		, bBorderColorBlack(false)
 		, bFlipGreenChannel(false)
+		, bApplyYCoCgBlockScale(false)
 		, bApplyKernelToTopMip(false)
 		, bRenormalizeTopMip(false)
 		, CompositeTextureMode(0 /*CTM_Disabled*/)
diff --git a/Engine/Source/Runtime/Engine/Classes/Engine/Texture.h b/Engine/Source/Runtime/Engine/Classes/Engine/Texture.h
--- a/Engine/Source/Runtime/Engine/Classes/Engine/Texture.h
+++ b/Engine/Source/Runtime/Engine/Classes/Engine/Texture.h
@@ -509,6 +509,7 @@ struct FTextureFormatSettings
 		: CompressionSettings(TC_Default)
 		, CompressionNoAlpha(false)
 		, CompressionNone(false)
+		, CompressionYCoCg(false)
 		, SRGB(false)
 	{}
 
@@ -521,6 +522,9 @@ struct FTextureFormatSettings
 	UPROPERTY()
 	uint8 CompressionNone : 1;
 
+	UPROPERTY()
+	uint8 CompressionYCoCg : 1;
+
 	UPROPERTY()
 	uint8 SRGB : 1;
 };
@@ -725,6 +729,10 @@ public:
 	UPROPERTY(EditAnywhere, BlueprintReadOnly, Category = Texture, AssetRegistrySearchable, AdvancedDisplay)
 	uint8 VirtualTextureStreaming : 1;
 
+	/** If true the texture stores YCoCg. Blue channel will be filled with a precision scale during compression. */
+	UPROPERTY()
+	uint8 CompressionYCoCg : 1;
+
 private:
 	/** Whether the async resource release process has already been kicked off or not */
 	UPROPERTY(transient)
diff --git a/Engine/Source/Runtime/Engine/Classes/VT/RuntimeVirtualTexture.h b/Engine/Source/Runtime/Engine/Classes/VT/RuntimeVirtualTexture.h
--- a/Engine/Source/Runtime/Engine/Classes/VT/RuntimeVirtualTexture.h
+++ b/Engine/Source/Runtime/Engine/Classes/VT/RuntimeVirtualTexture.h
@@ -115,6 +115,8 @@ public:
 	EPixelFormat GetLayerFormat(int32 LayerIndex) const;
 	/** Return true if the virtual texture layer should be sampled as sRGB */
 	bool IsLayerSRGB(int32 LayerIndex) const;
+	/** Return true if the virtual texture layer should be sampled as YCoCg */
+	bool IsLayerYCoCg(int32 LayerIndex) const;
 	/** Returns true if texture pages should be cleared before render */
 	bool GetClearTextures() const { return bClearTextures; }
 
diff --git a/Engine/Source/Runtime/Engine/Private/Texture.cpp b/Engine/Source/Runtime/Engine/Private/Texture.cpp
--- a/Engine/Source/Runtime/Engine/Private/Texture.cpp
+++ b/Engine/Source/Runtime/Engine/Private/Texture.cpp
@@ -87,6 +87,7 @@ UTexture::UTexture(const FObjectInitializer& ObjectInitializer)
 	ChromaKeyColor = FColorList::Magenta;
 	ChromaKeyThreshold = 1.0f / 255.0f;
 	VirtualTextureStreaming = 0;
+	CompressionYCoCg = 0;
 	
 #endif // #if WITH_EDITORONLY_DATA
 
@@ -1248,6 +1249,7 @@ void UTexture::GetDefaultFormatSettings(FTextureFormatSettings& OutSettings) con
 	OutSettings.CompressionSettings = CompressionSettings;
 	OutSettings.CompressionNone = CompressionNone;
 	OutSettings.CompressionNoAlpha = CompressionNoAlpha;
+	OutSettings.CompressionYCoCg = CompressionYCoCg;
 	OutSettings.SRGB = SRGB;
 }
 
@@ -1273,6 +1275,7 @@ void UTexture::SetLayerFormatSettings(int32 LayerIndex, const FTextureFormatSett
 		CompressionSettings = InSettings.CompressionSettings;
 		CompressionNone = InSettings.CompressionNone;
 		CompressionNoAlpha = InSettings.CompressionNoAlpha;
+		CompressionYCoCg = InSettings.CompressionYCoCg;
 		SRGB = InSettings.SRGB;
 	}
 	else
diff --git a/Engine/Source/Runtime/Engine/Private/TextureDerivedData.cpp b/Engine/Source/Runtime/Engine/Private/TextureDerivedData.cpp
--- a/Engine/Source/Runtime/Engine/Private/TextureDerivedData.cpp
+++ b/Engine/Source/Runtime/Engine/Private/TextureDerivedData.cpp
@@ -147,6 +147,7 @@ static void SerializeForKey(FArchive& Ar, const FTextureBuildSettings& Settings)
 		TempByte = Settings.bVirtualTextureEnableCompressZlib; Ar << TempByte;
 		TempByte = Settings.bVirtualTextureEnableCompressCrunch; Ar << TempByte;
 		TempByte = Settings.LossyCompressionAmount; Ar << TempByte; // Lossy compression currently only used by VT
+		TempByte = Settings.bApplyYCoCgBlockScale; Ar << TempByte; // YCoCg currently only used by VT
 	}
 }
 
@@ -310,6 +311,7 @@ static void FinalizeBuildSettingsForLayer(const UTexture& Texture, int32 LayerIn
 
 	OutSettings.bHDRSource = Texture.HasHDRSource(LayerIndex);
 	OutSettings.bSRGB = FormatSettings.SRGB;
+	OutSettings.bApplyYCoCgBlockScale = FormatSettings.CompressionYCoCg;
 
 	if (FormatSettings.CompressionSettings == TC_Displacementmap || FormatSettings.CompressionSettings == TC_DistanceFieldFont)
 	{
diff --git a/Engine/Source/Runtime/Engine/Private/VT/RuntimeVirtualTexture.cpp b/Engine/Source/Runtime/Engine/Private/VT/RuntimeVirtualTexture.cpp
--- a/Engine/Source/Runtime/Engine/Private/VT/RuntimeVirtualTexture.cpp
+++ b/Engine/Source/Runtime/Engine/Private/VT/RuntimeVirtualTexture.cpp
@@ -362,6 +362,11 @@ bool URuntimeVirtualTexture::IsLayerSRGB(int32 LayerIndex) const
 	return false;
 }
 
+bool URuntimeVirtualTexture::IsLayerYCoCg(int32 LayerIndex) const
+{
+	return MaterialType == ERuntimeVirtualTextureMaterialType::BaseColor_Normal_Specular_YCoCg && LayerIndex == 0;
+}
+
 int32 URuntimeVirtualTexture::GetEstimatedPageTableTextureMemoryKb() const
 {
 	//todo[vt]: Estimate memory usage
@@ -531,10 +536,12 @@ void URuntimeVirtualTexture::InitializeStreamingTexture(uint32 InSizeX, uint32 I
 		LayerFormats[Layer] = LayerFormat == PF_G16 ? TSF_G16 : TSF_BGRA8;
 
 		FTextureFormatSettings FormatSettings;
-		FormatSettings.SRGB = IsLayerSRGB(Layer);
-		FormatSettings.CompressionNone = LayerFormat == PF_B8G8R8A8 ||LayerFormat == PF_G16;
-		FormatSettings.CompressionNoAlpha = LayerFormat == PF_DXT1 || LayerFormat == PF_BC5;
 		FormatSettings.CompressionSettings = LayerFormat == PF_BC5 ? TC_Normalmap : TC_Default;
+		FormatSettings.CompressionNone = LayerFormat == PF_B8G8R8A8 || LayerFormat == PF_G16;
+		FormatSettings.CompressionNoAlpha = LayerFormat == PF_DXT1 || LayerFormat == PF_BC5;
+		FormatSettings.CompressionYCoCg = IsLayerYCoCg(Layer);
+		FormatSettings.SRGB = IsLayerSRGB(Layer);
+		
 		StreamingTexture->SetLayerFormatSettings(Layer, FormatSettings);
 	}
 
diff --git a/Engine/Source/Runtime/Engine/Private/VT/VirtualTextureDataBuilder.cpp b/Engine/Source/Runtime/Engine/Private/VT/VirtualTextureDataBuilder.cpp
--- a/Engine/Source/Runtime/Engine/Private/VT/VirtualTextureDataBuilder.cpp
+++ b/Engine/Source/Runtime/Engine/Private/VT/VirtualTextureDataBuilder.cpp
@@ -860,6 +860,7 @@ void FVirtualTextureDataBuilder::BuildSourcePixels(const FTextureSourceData& Sou
 			TBSettings.TextureFormatName = LayerData.FormatName;
 			TBSettings.bSRGB = BuildSettingsForLayer.bSRGB;
 			TBSettings.bUseLegacyGamma = BuildSettingsForLayer.bUseLegacyGamma;
+			TBSettings.bApplyYCoCgBlockScale = BuildSettingsForLayer.bApplyYCoCgBlockScale;
 
 			// Make sure the output of the texture builder is in the same gamma space as we expect it.
 			check(TBSettings.GetGammaSpace() == BuildSettingsForLayer.GetGammaSpace());
