Added Unfiltered MipGenSettings to generate mipmaps without averaging texels.

diff --git a/Engine/Source/Developer/TextureCompressor/Private/TextureCompressorModule.cpp b/Engine/Source/Developer/TextureCompressor/Private/TextureCompressorModule.cpp
--- a/Engine/Source/Developer/TextureCompressor/Private/TextureCompressorModule.cpp
+++ b/Engine/Source/Developer/TextureCompressor/Private/TextureCompressorModule.cpp
@@ -422,7 +422,8 @@ static void GenerateSharpenedMipB8G8R8A8Templ(
 	FVector4 AlphaThresholds,
 	const FImageKernel2D& Kernel,
 	uint32 ScaleFactor,
-	bool bSharpenWithoutColorShift )
+	bool bSharpenWithoutColorShift,
+	bool bUnfiltered)
 {
 	check( SourceImageData.SizeX == ScaleFactor * DestImageData.SizeX || DestImageData.SizeX == 1 );
 	check( SourceImageData.SizeY == ScaleFactor * DestImageData.SizeY || DestImageData.SizeY == 1 );
@@ -448,7 +449,11 @@ static void GenerateSharpenedMipB8G8R8A8Templ(
 
 			FLinearColor FilteredColor(0, 0, 0, 0);
 
-			if ( bSharpenWithoutColorShift )
+			if ( bUnfiltered )
+			{
+				FilteredColor = LookupSourceMip<AddressMode>(SourceImageData, SourceX + 0, SourceY + 0);
+			}
+			else if ( bSharpenWithoutColorShift )
 			{
 				FLinearColor SharpenedColor(0, 0, 0, 0);
 
@@ -537,26 +542,27 @@ static void GenerateSharpenedMipB8G8R8A8(
 	FVector4 AlphaThresholds,
 	const FImageKernel2D &Kernel,
 	uint32 ScaleFactor,
-	bool bSharpenWithoutColorShift
+	bool bSharpenWithoutColorShift,
+	bool bUnfiltered
 	)
 {
 	switch(AddressMode)
 	{
 	case MGTAM_Wrap:
-		GenerateSharpenedMipB8G8R8A8Templ<MGTAM_Wrap>(SourceImageData, DestImageData, bDitherMipMapAlpha, AlphaCoverages, AlphaThresholds, Kernel, ScaleFactor, bSharpenWithoutColorShift);
+		GenerateSharpenedMipB8G8R8A8Templ<MGTAM_Wrap>(SourceImageData, DestImageData, bDitherMipMapAlpha, AlphaCoverages, AlphaThresholds, Kernel, ScaleFactor, bSharpenWithoutColorShift, bUnfiltered);
 		break;
 	case MGTAM_Clamp:
-		GenerateSharpenedMipB8G8R8A8Templ<MGTAM_Clamp>(SourceImageData, DestImageData, bDitherMipMapAlpha, AlphaCoverages, AlphaThresholds, Kernel, ScaleFactor, bSharpenWithoutColorShift);
+		GenerateSharpenedMipB8G8R8A8Templ<MGTAM_Clamp>(SourceImageData, DestImageData, bDitherMipMapAlpha, AlphaCoverages, AlphaThresholds, Kernel, ScaleFactor, bSharpenWithoutColorShift, bUnfiltered);
 		break;
 	case MGTAM_BorderBlack:
-		GenerateSharpenedMipB8G8R8A8Templ<MGTAM_BorderBlack>(SourceImageData, DestImageData, bDitherMipMapAlpha, AlphaCoverages, AlphaThresholds, Kernel, ScaleFactor, bSharpenWithoutColorShift);
+		GenerateSharpenedMipB8G8R8A8Templ<MGTAM_BorderBlack>(SourceImageData, DestImageData, bDitherMipMapAlpha, AlphaCoverages, AlphaThresholds, Kernel, ScaleFactor, bSharpenWithoutColorShift, bUnfiltered);
 		break;
 	default:
 		check(0);
 	}
 
 	// For volume texture, do the average between the 2.
-	if (SourceImageData2.IsValid())
+	if (SourceImageData2.IsValid() && !bUnfiltered)
 	{
 		FImage Temp(DestImageData.SizeX, DestImageData.SizeY, 1, ERawImageFormat::RGBA32F);
 		FImageView2D TempImageData (Temp, 0);
@@ -564,13 +570,13 @@ static void GenerateSharpenedMipB8G8R8A8(
 		switch(AddressMode)
 		{
 		case MGTAM_Wrap:
-			GenerateSharpenedMipB8G8R8A8Templ<MGTAM_Wrap>(SourceImageData2, TempImageData, bDitherMipMapAlpha, AlphaCoverages, AlphaThresholds, Kernel, ScaleFactor, bSharpenWithoutColorShift);
+			GenerateSharpenedMipB8G8R8A8Templ<MGTAM_Wrap>(SourceImageData2, TempImageData, bDitherMipMapAlpha, AlphaCoverages, AlphaThresholds, Kernel, ScaleFactor, bSharpenWithoutColorShift, bUnfiltered);
 			break;
 		case MGTAM_Clamp:
-			GenerateSharpenedMipB8G8R8A8Templ<MGTAM_Clamp>(SourceImageData2, TempImageData, bDitherMipMapAlpha, AlphaCoverages, AlphaThresholds, Kernel, ScaleFactor, bSharpenWithoutColorShift);
+			GenerateSharpenedMipB8G8R8A8Templ<MGTAM_Clamp>(SourceImageData2, TempImageData, bDitherMipMapAlpha, AlphaCoverages, AlphaThresholds, Kernel, ScaleFactor, bSharpenWithoutColorShift, bUnfiltered);
 			break;
 		case MGTAM_BorderBlack:
-			GenerateSharpenedMipB8G8R8A8Templ<MGTAM_BorderBlack>(SourceImageData2, TempImageData, bDitherMipMapAlpha, AlphaCoverages, AlphaThresholds, Kernel, ScaleFactor, bSharpenWithoutColorShift);
+			GenerateSharpenedMipB8G8R8A8Templ<MGTAM_BorderBlack>(SourceImageData2, TempImageData, bDitherMipMapAlpha, AlphaCoverages, AlphaThresholds, Kernel, ScaleFactor, bSharpenWithoutColorShift, bUnfiltered);
 			break;
 		default:
 			check(0);
@@ -681,8 +687,8 @@ static void GenerateTopMip(const FImage& SrcImage, FImage& DestImage, const FTex
 			FVector4(0, 0, 0, 0),
 			KernelDownsample,
 			1,
-			Settings.bSharpenWithoutColorShift
-			);
+			Settings.bSharpenWithoutColorShift,
+			Settings.MipGenSettings == TMGS_Unfiltered);
 	}
 }
 
@@ -775,8 +781,8 @@ static void GenerateMipChain(
 				Settings.AlphaCoverageThresholds,
 				KernelDownsample,
 				2,
-				Settings.bSharpenWithoutColorShift
-				);
+				Settings.bSharpenWithoutColorShift,
+				Settings.MipGenSettings == TMGS_Unfiltered);
 
 			// generate IntermediateDstImage:
 			if ( Settings.bDownsampleWithAverage )
@@ -792,8 +798,8 @@ static void GenerateMipChain(
 					Settings.AlphaCoverageThresholds,
 					KernelSimpleAverage,
 					2,
-					Settings.bSharpenWithoutColorShift
-					);
+					Settings.bSharpenWithoutColorShift,
+					Settings.MipGenSettings == TMGS_Unfiltered);
 			}
 		}
 
diff --git a/Engine/Source/Runtime/Engine/Classes/Engine/TextureDefines.h b/Engine/Source/Runtime/Engine/Classes/Engine/TextureDefines.h
--- a/Engine/Source/Runtime/Engine/Classes/Engine/TextureDefines.h
+++ b/Engine/Source/Runtime/Engine/Classes/Engine/TextureDefines.h
@@ -58,11 +58,11 @@ enum TextureGroup
 	/** Impostor Color Textures*/
 	TEXTUREGROUP_Impostor UMETA(DisplayName="ini:Impostor Color"),
 	/** Impostor Normal and Depth, use default compression*/
-	TEXTUREGROUP_ImpostorNormalDepth UMETA(DisplayName = "ini:Impostor Normal and Depth"),
+	TEXTUREGROUP_ImpostorNormalDepth UMETA(DisplayName="ini:Impostor Normal and Depth"),
 	/** 8 bit data stored in textures */
-	TEXTUREGROUP_8BitData UMETA(DisplayName = "ini:8 Bit Data"),
+	TEXTUREGROUP_8BitData UMETA(DisplayName="ini:8 Bit Data"),
 	/** 16 bit data stored in textures */
-	TEXTUREGROUP_16BitData UMETA(DisplayName = "ini:16 Bit Data"),
+	TEXTUREGROUP_16BitData UMETA(DisplayName="ini:16 Bit Data"),
 	/** Project specific group, rename in Engine.ini, [EnumRemap] TEXTUREGROUP_Project**.DisplayName=My Fun Group */
 	TEXTUREGROUP_Project01 UMETA(DisplayName="ini:Project Group 01"),
 	TEXTUREGROUP_Project02 UMETA(DisplayName="ini:Project Group 02"),
@@ -87,25 +87,25 @@ enum TextureMipGenSettings
 	/** 8x8 with sharpening: 0=no sharpening but better quality which is softer, 1=little, 5=medium, 10=extreme. */
 	TMGS_Sharpen0 UMETA(DisplayName="Sharpen0"),
 	/** 8x8 with sharpening: 0=no sharpening but better quality which is softer, 1=little, 5=medium, 10=extreme. */
-	TMGS_Sharpen1 UMETA(DisplayName = "Sharpen1"),
+	TMGS_Sharpen1 UMETA(DisplayName="Sharpen1"),
 	/** 8x8 with sharpening: 0=no sharpening but better quality which is softer, 1=little, 5=medium, 10=extreme. */
-	TMGS_Sharpen2 UMETA(DisplayName = "Sharpen2"),
+	TMGS_Sharpen2 UMETA(DisplayName="Sharpen2"),
 	/** 8x8 with sharpening: 0=no sharpening but better quality which is softer, 1=little, 5=medium, 10=extreme. */
-	TMGS_Sharpen3 UMETA(DisplayName = "Sharpen3"),
+	TMGS_Sharpen3 UMETA(DisplayName="Sharpen3"),
 	/** 8x8 with sharpening: 0=no sharpening but better quality which is softer, 1=little, 5=medium, 10=extreme. */
-	TMGS_Sharpen4 UMETA(DisplayName = "Sharpen4"),
+	TMGS_Sharpen4 UMETA(DisplayName="Sharpen4"),
 	/** 8x8 with sharpening: 0=no sharpening but better quality which is softer, 1=little, 5=medium, 10=extreme. */
-	TMGS_Sharpen5 UMETA(DisplayName = "Sharpen5"),
+	TMGS_Sharpen5 UMETA(DisplayName="Sharpen5"),
 	/** 8x8 with sharpening: 0=no sharpening but better quality which is softer, 1=little, 5=medium, 10=extreme. */
-	TMGS_Sharpen6 UMETA(DisplayName = "Sharpen6"),
+	TMGS_Sharpen6 UMETA(DisplayName="Sharpen6"),
 	/** 8x8 with sharpening: 0=no sharpening but better quality which is softer, 1=little, 5=medium, 10=extreme. */
-	TMGS_Sharpen7 UMETA(DisplayName = "Sharpen7"),
+	TMGS_Sharpen7 UMETA(DisplayName="Sharpen7"),
 	/** 8x8 with sharpening: 0=no sharpening but better quality which is softer, 1=little, 5=medium, 10=extreme. */
-	TMGS_Sharpen8 UMETA(DisplayName = "Sharpen8"),
+	TMGS_Sharpen8 UMETA(DisplayName="Sharpen8"),
 	/** 8x8 with sharpening: 0=no sharpening but better quality which is softer, 1=little, 5=medium, 10=extreme. */
-	TMGS_Sharpen9 UMETA(DisplayName = "Sharpen9"),
+	TMGS_Sharpen9 UMETA(DisplayName="Sharpen9"),
 	/** 8x8 with sharpening: 0=no sharpening but better quality which is softer, 1=little, 5=medium, 10=extreme. */
-	TMGS_Sharpen10 UMETA(DisplayName = "Sharpen10"),
+	TMGS_Sharpen10 UMETA(DisplayName="Sharpen10"),
 	TMGS_NoMipmaps UMETA(DisplayName="NoMipmaps"),
 	/** Do not touch existing mip chain as it contains generated data. */
 	TMGS_LeaveExistingMips UMETA(DisplayName="LeaveExistingMips"),
@@ -114,7 +114,9 @@ enum TextureMipGenSettings
 	TMGS_Blur2 UMETA(DisplayName="Blur2"),
 	TMGS_Blur3 UMETA(DisplayName="Blur3"),
 	TMGS_Blur4 UMETA(DisplayName="Blur4"),
-	TMGS_Blur5 UMETA(DisplayName="Blur5"),
+	TMGS_Blur5 UMETA(DisplayName = "Blur5"),
+	/** Use the first texel of each 2x2 (or 2x2x2) group. */
+	TMGS_Unfiltered UMETA(DisplayName = "Unfiltered"),
 	TMGS_MAX,
 
 	// Note: These are serialized as as raw values in the texture DDC key, so additional entries
