PR #6292: Engine changes for USD plugin (Contributed by [REDACTED_PERSON])

diff --git a/Engine/Source/Developer/TextureCompressor/Private/TextureCompressorModule.cpp b/Engine/Source/Developer/TextureCompressor/Private/TextureCompressorModule.cpp
--- a/Engine/Source/Developer/TextureCompressor/Private/TextureCompressorModule.cpp
+++ b/Engine/Source/Developer/TextureCompressor/Private/TextureCompressorModule.cpp
@@ -985,12 +985,7 @@ static int32 ComputeLongLatCubemapExtents(const FImage& SrcImage, const int32 Ma
 	return FMath::Clamp(1 << FMath::FloorLog2(SrcImage.SizeX / 2), 32, MaxCubemapTextureResolution);
 }
 
-/**
- * Generates the base cubemap mip from a longitude-latitude 2D image.
- * @param OutMip - The output mip.
- * @param SrcImage - The source longlat image.
- */
-static void GenerateBaseCubeMipFromLongitudeLatitude2D(FImage* OutMip, const FImage& SrcImage, const int32 MaxCubemapTextureResolution)
+void ITextureCompressorModule::GenerateBaseCubeMipFromLongitudeLatitude2D(FImage* OutMip, const FImage& SrcImage, const int32 MaxCubemapTextureResolution)
 {
 	FImage LongLatImage;
 	SrcImage.CopyTo(LongLatImage, ERawImageFormat::RGBA32F, EGammaSpace::Linear);
@@ -1289,13 +1284,7 @@ static void GenerateAngularFilteredMip(FImage* DestMip, FImage& SrcMip, float Co
 	}
 }
 
-/**
- * Generates angularly filtered mips.
- * @param InOutMipChain - The mip chain to angularly filter.
- * @param NumMips - The number of mips the chain should have.
- * @param DiffuseConvolveMipLevel - The mip level that contains the diffuse convolution.
- */
-static void GenerateAngularFilteredMips(TArray<FImage>& InOutMipChain, int32 NumMips, uint32 DiffuseConvolveMipLevel)
+void ITextureCompressorModule::GenerateAngularFilteredMips(TArray<FImage>& InOutMipChain, int32 NumMips, uint32 DiffuseConvolveMipLevel)
 {
 	TArray<FImage> SrcMipChain;
 	Exchange(SrcMipChain, InOutMipChain);
diff --git a/Engine/Source/Developer/TextureCompressor/Public/TextureCompressorModule.h b/Engine/Source/Developer/TextureCompressor/Public/TextureCompressorModule.h
--- a/Engine/Source/Developer/TextureCompressor/Public/TextureCompressorModule.h
+++ b/Engine/Source/Developer/TextureCompressor/Public/TextureCompressorModule.h
@@ -270,4 +270,20 @@ public:
      * @param	InBuildSettings	Image build settings
      */
 	TEXTURECOMPRESSOR_API static void AdjustImageColors(FImage& Image, const FTextureBuildSettings& InBuildSettings);
+
+	/**
+	 * Generates the base cubemap mip from a longitude-latitude 2D image.
+	 * @param OutMip - The output mip.
+	 * @param SrcImage - The source longlat image.
+	 */
+	TEXTURECOMPRESSOR_API static void GenerateBaseCubeMipFromLongitudeLatitude2D(FImage* OutMip, const FImage& SrcImage, const int32 MaxCubemapTextureResolution);
+
+
+	/**
+	 * Generates angularly filtered mips.
+	 * @param InOutMipChain - The mip chain to angularly filter.
+	 * @param NumMips - The number of mips the chain should have.
+	 * @param DiffuseConvolveMipLevel - The mip level that contains the diffuse convolution.
+	 */
+	TEXTURECOMPRESSOR_API static void GenerateAngularFilteredMips(TArray<FImage>& InOutMipChain, int32 NumMips, uint32 DiffuseConvolveMipLevel);
 };
diff --git a/Engine/Source/Runtime/Engine/Classes/Materials/MaterialExpressionCameraPositionWS.h b/Engine/Source/Runtime/Engine/Classes/Materials/MaterialExpressionCameraPositionWS.h
--- a/Engine/Source/Runtime/Engine/Classes/Materials/MaterialExpressionCameraPositionWS.h
+++ b/Engine/Source/Runtime/Engine/Classes/Materials/MaterialExpressionCameraPositionWS.h
@@ -8,7 +8,7 @@
 #include "Materials/MaterialExpression.h"
 #include "MaterialExpressionCameraPositionWS.generated.h"
 
-UCLASS(collapsecategories, hidecategories=Object)
+UCLASS(MinimalAPI, collapsecategories, hidecategories=Object)
 class UMaterialExpressionCameraPositionWS : public UMaterialExpression
 {
 	GENERATED_UCLASS_BODY()
diff --git a/Engine/Source/Runtime/Engine/Classes/Materials/MaterialExpressionSingleLayerWaterMaterialOutput.h b/Engine/Source/Runtime/Engine/Classes/Materials/MaterialExpressionSingleLayerWaterMaterialOutput.h
--- a/Engine/Source/Runtime/Engine/Classes/Materials/MaterialExpressionSingleLayerWaterMaterialOutput.h
+++ b/Engine/Source/Runtime/Engine/Classes/Materials/MaterialExpressionSingleLayerWaterMaterialOutput.h
@@ -8,7 +8,7 @@
 #include "MaterialExpressionSingleLayerWaterMaterialOutput.generated.h"
 
 /** Material output expression for writing single layer water volume material properties. */
-UCLASS(collapsecategories, hidecategories = Object)
+UCLASS(MinimalAPI, collapsecategories, hidecategories = Object)
 class UMaterialExpressionSingleLayerWaterMaterialOutput : public UMaterialExpressionCustomOutput
 {
 	GENERATED_UCLASS_BODY()
