Add more info to failed composite texture
#rnx

diff --git a/Engine/Source/Developer/TextureCompressor/Private/TextureCompressorModule.cpp b/Engine/Source/Developer/TextureCompressor/Private/TextureCompressorModule.cpp
--- a/Engine/Source/Developer/TextureCompressor/Private/TextureCompressorModule.cpp
+++ b/Engine/Source/Developer/TextureCompressor/Private/TextureCompressorModule.cpp
@@ -2331,7 +2331,14 @@ private:
 		if( RoughnessSourceMips[RoughnessSourceMips.Num() - MinLevel].SizeX != NormalSourceMips[NormalSourceMips.Num() - MinLevel].SizeX || 
 			RoughnessSourceMips[RoughnessSourceMips.Num() - MinLevel].SizeY != NormalSourceMips[NormalSourceMips.Num() - MinLevel].SizeY )
 		{
-			//incomplete mip chain or mismatched dimensions so bail
+			UE_LOG(LogTextureCompressor, Warning, TEXT("Couldn't apply composite texture as RoughnessSourceMips (mip %d, %d x %d) doesn't match NormalSourceMips (mip %d, %d x %d); mipchain might be mismatched/incomplete"),
+				RoughnessSourceMips.Num() - MinLevel,
+				RoughnessSourceMips[RoughnessSourceMips.Num() - MinLevel].SizeX,
+				RoughnessSourceMips[RoughnessSourceMips.Num() - MinLevel].SizeY,
+				NormalSourceMips.Num() - MinLevel,
+				NormalSourceMips[NormalSourceMips.Num() - MinLevel].SizeX,
+				NormalSourceMips[NormalSourceMips.Num() - MinLevel].SizeY
+				);
 			return false;
 		}
 
