Add an option to downscale textures without mips. Accepts values between 1.0 and 8.0.
Can be overriden per-platform, for a specifc texture or texture group.
Downscaling will keep texure size divisible by 4 so they can still be compressed

diff --git a/Engine/Source/Developer/TextureCompressor/Private/TextureCompressorModule.cpp b/Engine/Source/Developer/TextureCompressor/Private/TextureCompressorModule.cpp
--- a/Engine/Source/Developer/TextureCompressor/Private/TextureCompressorModule.cpp
+++ b/Engine/Source/Developer/TextureCompressor/Private/TextureCompressorModule.cpp
@@ -692,6 +692,192 @@ static void GenerateTopMip(const FImage& SrcImage, FImage& DestImage, const FTex
 	}
 }
 
+static FLinearColor LookupSourceMipBilinear(const FImageView2D& SourceImageData, float X, float Y)
+{
+	X = FMath::Clamp(X, 0.f, SourceImageData.SizeX - 1.f);
+	Y = FMath::Clamp(Y, 0.f, SourceImageData.SizeY - 1.f);
+	int32 IntX0 = FMath::FloorToInt(X);
+	int32 IntY0 = FMath::FloorToInt(Y);
+	float FractX = X - IntX0;
+	float FractY = Y - IntY0;
+	int32 IntX1 = FMath::Min(IntX0+1, SourceImageData.SizeX-1);
+	int32 IntY1 = FMath::Min(IntY0+1, SourceImageData.SizeY-1);
+	
+	FLinearColor Sample00 = SourceImageData.Access(IntX0,IntY0);
+	FLinearColor Sample10 = SourceImageData.Access(IntX1,IntY0);
+	FLinearColor Sample01 = SourceImageData.Access(IntX0,IntY1);
+	FLinearColor Sample11 = SourceImageData.Access(IntX1,IntY1);
+	FLinearColor Sample0 = FMath::Lerp(Sample00, Sample10, FractX);
+	FLinearColor Sample1 = FMath::Lerp(Sample01, Sample11, FractX);
+		
+	return FMath::Lerp(Sample0, Sample1, FractY);
+}
+
+struct FTextureDownscaleSettings
+{
+	int32 BlockSize;
+	float Downscale;
+	uint8 DownscaleOptions;
+	bool bDitherMipMapAlpha;
+};
+
+static void DownscaleImage(const FImage& SrcImage, FImage& DstImage, const FTextureDownscaleSettings& Settings)
+{
+	if (Settings.Downscale <= 1.f)
+	{
+		return;
+	}
+	
+	float Downscale = FMath::Clamp(Settings.Downscale, 1.f, 8.f);
+	int32 FinalSizeX = FMath::CeilToInt(SrcImage.SizeX / Downscale);
+	int32 FinalSizeY = FMath::CeilToInt(SrcImage.SizeY / Downscale);
+
+	// compute final size respecting image block size
+	if (Settings.BlockSize > 1 
+		&& SrcImage.SizeX % Settings.BlockSize == 0 
+		&& SrcImage.SizeY % Settings.BlockSize == 0)
+	{
+		int32 NumBlocksX = SrcImage.SizeX / Settings.BlockSize;
+		int32 NumBlocksY = SrcImage.SizeY / Settings.BlockSize;
+		int32 GCD = FMath::GreatestCommonDivisor(NumBlocksX, NumBlocksY);
+		int32 RatioX = NumBlocksX/GCD;
+		int32 RatioY = NumBlocksY/GCD;
+		int32 FinalNumBlocksX = (int32)FMath::GridSnap((float)FinalSizeX/Settings.BlockSize, (float)RatioX);
+		int32 FinalNumBlocksY = FinalNumBlocksX/RatioX*RatioY;
+		FinalSizeX = FinalNumBlocksX*Settings.BlockSize;
+		FinalSizeY = FinalNumBlocksY*Settings.BlockSize;
+	}
+
+	Downscale = (float)SrcImage.SizeX / FinalSizeX;
+		
+	FImage Image0;
+	FImage Image1;
+	FImage* ImageChain[2] = {&const_cast<FImage&>(SrcImage), &Image1};
+	bool bUnfiltered = Settings.DownscaleOptions == (uint8)ETextureDownscaleOptions::Unfiltered;
+	
+	// Scaledown using 2x2 average, use user specified filtering only for last iteration
+	FImageKernel2D AvgKernel;
+	AvgKernel.BuildSeparatableGaussWithSharpen(2);
+	int32 NumIterations = 0;
+	while(Downscale > 2.0f)
+	{
+		int32 DstSizeX = ImageChain[0]->SizeX / 2;
+		int32 DstSizeY = ImageChain[0]->SizeY / 2;
+		ImageChain[1]->Init(DstSizeX, DstSizeY, ImageChain[0]->NumSlices, ImageChain[0]->Format, ImageChain[0]->GammaSpace);
+
+		FImageView2D SrcImageData(*ImageChain[0], 0);
+		FImageView2D DstImageData(*ImageChain[1], 0);
+		GenerateSharpenedMipB8G8R8A8Templ<MGTAM_Clamp>(
+			SrcImageData, 
+			DstImageData, 
+			Settings.bDitherMipMapAlpha, 
+			FVector4(0, 0, 0, 0),
+			FVector4(0, 0, 0, 0), 
+			AvgKernel, 
+			2, 
+			false,
+			bUnfiltered);
+
+		if (NumIterations == 0)
+		{
+			ImageChain[0] = &Image0;
+		}
+		Swap(ImageChain[0], ImageChain[1]);
+		
+		NumIterations++;
+		Downscale/= 2.f;
+	}
+
+	if (ImageChain[0]->SizeX == FinalSizeX &&
+		ImageChain[0]->SizeY == FinalSizeY)
+	{
+		ImageChain[0]->CopyTo(DstImage, ImageChain[0]->Format, ImageChain[0]->GammaSpace);
+		return;
+	}
+	
+	int32 KernelSize = 2;
+	float Sharpening = 0.0f;
+	if (Settings.DownscaleOptions >= (uint8)ETextureDownscaleOptions::Sharpen0 && Settings.DownscaleOptions <= (uint8)ETextureDownscaleOptions::Sharpen10)
+	{
+		// 0 .. 2.0f
+		Sharpening = ((int32)Settings.DownscaleOptions - (int32)ETextureDownscaleOptions::Sharpen0) * 0.2f;
+		KernelSize = 8;
+	}
+	
+	bool bBilinear = Settings.DownscaleOptions == (uint8)ETextureDownscaleOptions::SimpleAverage;
+	
+	FImageKernel2D KernelSharpen;
+	KernelSharpen.BuildSeparatableGaussWithSharpen(KernelSize, Sharpening);
+	const int32 KernelCenter = (int32)KernelSharpen.GetFilterTableSize() / 2 - 1;
+		
+	ImageChain[1] = &DstImage;
+	if (ImageChain[0] == ImageChain[1])
+	{
+		ImageChain[0]->CopyTo(Image0, ImageChain[0]->Format, ImageChain[0]->GammaSpace);
+		ImageChain[0] = &Image0;
+	}
+	
+	// Set up a random number stream for dithering.
+	FRandomStream RandomStream(0);
+	ImageChain[1]->Init(FinalSizeX, FinalSizeY, ImageChain[0]->NumSlices, ImageChain[0]->Format, ImageChain[0]->GammaSpace);
+	Downscale = (float)ImageChain[0]->SizeX / FinalSizeX;
+
+	FImageView2D SrcImageData(*ImageChain[0], 0);
+	FImageView2D DstImageData(*ImageChain[1], 0);
+					
+	for (int32 Y = 0; Y < FinalSizeY; ++Y)
+	{
+		float SourceY = Y * Downscale;
+		int32 IntSourceY = FMath::RoundToInt(SourceY);
+		
+		for (int32 X = 0; X < FinalSizeX; ++X)
+		{
+			float SourceX = X * Downscale;
+			int32 IntSourceX = FMath::RoundToInt(SourceX);
+
+			FLinearColor FilteredColor(0,0,0,0);
+
+			if (bUnfiltered)
+			{
+				FilteredColor = LookupSourceMip<MGTAM_Clamp>(SrcImageData, IntSourceX, IntSourceY);
+			}
+			else if(bBilinear)
+			{
+				FilteredColor = LookupSourceMipBilinear(SrcImageData, SourceX, SourceY);
+			}
+			else
+			{
+				for (uint32 KernelY = 0; KernelY < KernelSharpen.GetFilterTableSize();  ++KernelY)
+				{
+					for (uint32 KernelX = 0; KernelX < KernelSharpen.GetFilterTableSize();  ++KernelX)
+					{
+						float Weight = KernelSharpen.GetAt(KernelX, KernelY);
+						FLinearColor Sample = LookupSourceMipBilinear(SrcImageData, SourceX + KernelX - KernelCenter, SourceY + KernelY - KernelCenter);
+						FilteredColor += Weight	* Sample;
+					}
+				}
+			}
+
+			if (Settings.bDitherMipMapAlpha)
+			{
+				// Dither the alpha of any pixel which passes an alpha threshold test.
+				const int32 DitherAlphaThreshold = 5.0f / 255.0f;
+				const float MinRandomAlpha = 85.0f;
+				const float MaxRandomAlpha = 255.0f;
+
+				if (FilteredColor.A > DitherAlphaThreshold)
+				{
+					FilteredColor.A = FMath::TruncToInt(FMath::Lerp(MinRandomAlpha, MaxRandomAlpha, RandomStream.GetFraction()));
+				}
+			}
+
+			// Set the destination pixel.
+			FLinearColor& DestColor = DstImageData.Access(X, Y);
+			DestColor = FilteredColor;
+		}
+	}
+}
+
 void ITextureCompressorModule::GenerateMipChain(
 	const FTextureBuildSettings& Settings,
 	const FImage& BaseImage,
@@ -2293,6 +2479,17 @@ private:
 				}				
 			}
 
+			if (BuildSettings.Downscale > 1.f)
+			{
+				FTextureDownscaleSettings DownscaleSettings;
+				DownscaleSettings.Downscale = BuildSettings.Downscale;
+				DownscaleSettings.DownscaleOptions = BuildSettings.DownscaleOptions;
+				DownscaleSettings.bDitherMipMapAlpha = BuildSettings.bDitherMipMapAlpha;
+				DownscaleSettings.BlockSize = 4;
+		
+				DownscaleImage(*Mip, *Mip, DownscaleSettings);
+			}
+
 			// Apply color adjustments
 			AdjustImageColors(*Mip, BuildSettings);
 			if (BuildSettings.bComputeBokehAlpha)
diff --git a/Engine/Source/Developer/TextureCompressor/Public/TextureCompressorModule.h b/Engine/Source/Developer/TextureCompressor/Public/TextureCompressorModule.h
--- a/Engine/Source/Developer/TextureCompressor/Public/TextureCompressorModule.h
+++ b/Engine/Source/Developer/TextureCompressor/Public/TextureCompressorModule.h
@@ -153,6 +153,10 @@ struct FTextureBuildSettings
 	int32 CompressionQuality;
 	/** ETextureLossyCompressionAmount */
 	int32 LossyCompressionAmount;
+	/** Values > 1.0 will scale down source texture. Ignored for textures with mips */
+	float Downscale;
+	/** ETextureDownscaleOptions */
+	uint8 DownscaleOptions;
 	/** TextureAddress, opaque to avoid dependencies on engine headers. How to address the texture (clamp, wrap, ...) for virtual textures this is baked into the build data, for regular textures this is ignored. */
 	int32 VirtualAddressingModeX;
 	int32 VirtualAddressingModeY;
@@ -208,6 +212,8 @@ struct FTextureBuildSettings
 		, ChromaKeyThreshold(1.0f / 255.0f)
 		, CompressionQuality(-1)
 		, LossyCompressionAmount(0)
+		, Downscale(0.0)
+		, DownscaleOptions(0)
 		, VirtualAddressingModeX(0)
 		, VirtualAddressingModeY(0)
 		, VirtualTextureTileSize(0)
diff --git a/Engine/Source/Runtime/Engine/Classes/Engine/Texture.h b/Engine/Source/Runtime/Engine/Classes/Engine/Texture.h
--- a/Engine/Source/Runtime/Engine/Classes/Engine/Texture.h
+++ b/Engine/Source/Runtime/Engine/Classes/Engine/Texture.h
@@ -16,6 +16,7 @@
 #include "MaterialShared.h"
 #include "TextureResource.h"
 #include "Engine/StreamableRenderAsset.h"
+#include "PerPlatformProperties.h"
 #include "Texture.generated.h"
 
 class ITargetPlatform;
@@ -798,6 +799,19 @@ public:
 	UPROPERTY(EditAnywhere, BlueprintReadWrite, Category=LevelOfDetail, meta=(DisplayName="Texture Group"), AssetRegistrySearchable)
 	TEnumAsByte<enum TextureGroup> LODGroup;
 
+	/** Downscale source texture, applied only to textures without mips 
+	 * 0.0 - use scale value from texture group
+	 * 1.0 - do not scale texture
+	 * > 1.0 - scale texure
+	 */
+	UPROPERTY(EditAnywhere, Category=LevelOfDetail, AdvancedDisplay, meta=(ClampMin="0.0", ClampMax="8.0"))
+	FPerPlatformFloat Downscale;
+
+	/** Texture downscaling options */
+	UPROPERTY(EditAnywhere, Category=LevelOfDetail, AdvancedDisplay)
+	ETextureDownscaleOptions DownscaleOptions;
+
+	
 	/** This should be unchecked if using alpha channels individually as masks. */
 	UPROPERTY(EditAnywhere, BlueprintReadWrite, Category=Texture, meta=(DisplayName="sRGB"), AssetRegistrySearchable)
 	uint8 SRGB:1;
diff --git a/Engine/Source/Runtime/Engine/Classes/Engine/TextureDefines.h b/Engine/Source/Runtime/Engine/Classes/Engine/TextureDefines.h
--- a/Engine/Source/Runtime/Engine/Classes/Engine/TextureDefines.h
+++ b/Engine/Source/Runtime/Engine/Classes/Engine/TextureDefines.h
@@ -171,3 +171,36 @@ enum class ETextureMipLoadOptions : uint8
 	// Load only the first mip.
 	OnlyFirstMip,
 };
+
+UENUM()
+enum class ETextureDownscaleOptions : uint8
+{
+	/** Fallback to the "texture group" settings */
+	Default,
+	/** Unfiltered */
+	Unfiltered,
+	/** Average, default for the "texture group" */
+	SimpleAverage,
+	/** 8x8 with sharpening: 0=no sharpening but better quality which is softer, 1=little, 5=medium, 10=extreme. */
+	Sharpen0,
+	/** 8x8 with sharpening: 0=no sharpening but better quality which is softer, 1=little, 5=medium, 10=extreme. */
+	Sharpen1,
+	/** 8x8 with sharpening: 0=no sharpening but better quality which is softer, 1=little, 5=medium, 10=extreme. */
+	Sharpen2,
+	/** 8x8 with sharpening: 0=no sharpening but better quality which is softer, 1=little, 5=medium, 10=extreme. */
+	Sharpen3,
+	/** 8x8 with sharpening: 0=no sharpening but better quality which is softer, 1=little, 5=medium, 10=extreme. */
+	Sharpen4,
+	/** 8x8 with sharpening: 0=no sharpening but better quality which is softer, 1=little, 5=medium, 10=extreme. */
+	Sharpen5,
+	/** 8x8 with sharpening: 0=no sharpening but better quality which is softer, 1=little, 5=medium, 10=extreme. */
+	Sharpen6,
+	/** 8x8 with sharpening: 0=no sharpening but better quality which is softer, 1=little, 5=medium, 10=extreme. */
+	Sharpen7,
+	/** 8x8 with sharpening: 0=no sharpening but better quality which is softer, 1=little, 5=medium, 10=extreme. */
+	Sharpen8,
+	/** 8x8 with sharpening: 0=no sharpening but better quality which is softer, 1=little, 5=medium, 10=extreme. */
+	Sharpen9,
+	/** 8x8 with sharpening: 0=no sharpening but better quality which is softer, 1=little, 5=medium, 10=extreme. */
+	Sharpen10
+};
diff --git a/Engine/Source/Runtime/Engine/Classes/Engine/TextureLODSettings.h b/Engine/Source/Runtime/Engine/Classes/Engine/TextureLODSettings.h
--- a/Engine/Source/Runtime/Engine/Classes/Engine/TextureLODSettings.h
+++ b/Engine/Source/Runtime/Engine/Classes/Engine/TextureLODSettings.h
@@ -36,6 +36,8 @@ struct ENGINE_API FTextureLODGroup
 		, MipFilter(NAME_Point)
 		, MipLoadOptions(ETextureMipLoadOptions::AllMips)
 		, DuplicateNonOptionalMips(false)
+		, Downscale(1.0)
+		, DownscaleOptions(ETextureDownscaleOptions::SimpleAverage)
 	{
 		SetupGroup();
 	}
@@ -105,6 +107,12 @@ struct ENGINE_API FTextureLODGroup
 	UPROPERTY()
 	bool DuplicateNonOptionalMips;
 
+	UPROPERTY()
+	float Downscale;
+
+	UPROPERTY()
+	ETextureDownscaleOptions DownscaleOptions;
+	
 	void SetupGroup();
 };
 
@@ -154,6 +162,7 @@ public:
 
 #if WITH_EDITORONLY_DATA
 	void GetMipGenSettings( const UTexture& Texture, TextureMipGenSettings& OutMipGenSettings, float& OutSharpen, uint32& OutKernelSize, bool& bOutDownsampleWithAverage, bool& bOutSharpenWithoutColorShift, bool &bOutBorderColorBlack ) const;
+	void GetDownscaleOptions(const UTexture& Texture, const ITargetPlatform& CurrentPlatform, float& Downscale, ETextureDownscaleOptions& DownscaleOptions) const;
 #endif // #if WITH_EDITORONLY_DATA
 
 	/**
diff --git a/Engine/Source/Runtime/Engine/Private/Texture.cpp b/Engine/Source/Runtime/Engine/Private/Texture.cpp
--- a/Engine/Source/Runtime/Engine/Private/Texture.cpp
+++ b/Engine/Source/Runtime/Engine/Private/Texture.cpp
@@ -94,7 +94,8 @@ UTexture::UTexture(const FObjectInitializer& ObjectInitializer)
 	ChromaKeyThreshold = 1.0f / 255.0f;
 	VirtualTextureStreaming = 0;
 	CompressionYCoCg = 0;
-	
+	Downscale = 0.f;
+	DownscaleOptions = ETextureDownscaleOptions::Default;
 #endif // #if WITH_EDITORONLY_DATA
 
 	if (FApp::CanEverRender() && !IsTemplate())
diff --git a/Engine/Source/Runtime/Engine/Private/TextureDerivedData.cpp b/Engine/Source/Runtime/Engine/Private/TextureDerivedData.cpp
--- a/Engine/Source/Runtime/Engine/Private/TextureDerivedData.cpp
+++ b/Engine/Source/Runtime/Engine/Private/TextureDerivedData.cpp
@@ -150,6 +150,13 @@ static void SerializeForKey(FArchive& Ar, const FTextureBuildSettings& Settings)
 		TempByte = Settings.LossyCompressionAmount; Ar << TempByte; // Lossy compression currently only used by VT
 		TempByte = Settings.bApplyYCoCgBlockScale; Ar << TempByte; // YCoCg currently only used by VT
 	}
+
+	// Avoid changing key if texture is not being downscaled
+	if (Settings.Downscale > 1.0)
+	{
+		TempFloat = Settings.Downscale; Ar << TempFloat;
+		TempByte = Settings.DownscaleOptions; Ar << TempByte;
+	}
 }
 
 /**
@@ -332,11 +339,13 @@ static void FinalizeBuildSettingsForLayer(const UTexture& Texture, int32 LayerIn
 static void GetTextureBuildSettings(
 	const UTexture& Texture,
 	const UTextureLODSettings& TextureLODSettings,
-	bool bPlatformSupportsTextureStreaming,
-	bool bPlatformSupportsVirtualTextureStreaming,
+	const ITargetPlatform& CurrentPlatform,
 	FTextureBuildSettings& OutBuildSettings
 	)
 {
+	const bool bPlatformSupportsTextureStreaming = CurrentPlatform.SupportsFeature(ETargetPlatformFeatures::TextureStreaming);
+	const bool bPlatformSupportsVirtualTextureStreaming = CurrentPlatform.SupportsFeature(ETargetPlatformFeatures::VirtualTextureStreaming);
+		
 	OutBuildSettings.ColorAdjustment.AdjustBrightness = Texture.AdjustBrightness;
 	OutBuildSettings.ColorAdjustment.AdjustBrightnessCurve = Texture.AdjustBrightnessCurve;
 	OutBuildSettings.ColorAdjustment.AdjustVibrance = Texture.AdjustVibrance;
@@ -429,6 +438,13 @@ static void GetTextureBuildSettings(
 	// TODO - get default value from config/CVAR/LODGroup?
 	OutBuildSettings.LossyCompressionAmount = (Texture.LossyCompressionAmount == TLCA_Default) ? TLCA_Lowest : Texture.LossyCompressionAmount.GetValue();
 
+	OutBuildSettings.Downscale = 1.0f;
+	if (MipGenSettings == TMGS_NoMipmaps && 
+		Texture.IsA(UTexture2D::StaticClass()))	// TODO: support more texture types
+	{
+		TextureLODSettings.GetDownscaleOptions(Texture, CurrentPlatform, OutBuildSettings.Downscale, (ETextureDownscaleOptions&)OutBuildSettings.DownscaleOptions);
+	}
+	
 	// For virtual texturing we take the address mode into consideration
 	if (OutBuildSettings.bVirtualStreamable)
 	{
@@ -499,11 +515,8 @@ static void GetBuildSettingsForRunningPlatform(
 		check(CurrentPlatform != NULL);
 
 		const UTextureLODSettings* LODSettings = (UTextureLODSettings*)UDeviceProfileManager::Get().FindProfile(CurrentPlatform->PlatformName());
-		const bool bPlatformSupportsTextureStreaming = CurrentPlatform->SupportsFeature(ETargetPlatformFeatures::TextureStreaming);
-		const bool bPlatformSupportsVirtualTextureStreaming = CurrentPlatform->SupportsFeature(ETargetPlatformFeatures::VirtualTextureStreaming);
-
 		FTextureBuildSettings SourceBuildSettings;
-		GetTextureBuildSettings(Texture, *LODSettings, bPlatformSupportsTextureStreaming, bPlatformSupportsVirtualTextureStreaming, SourceBuildSettings);
+		GetTextureBuildSettings(Texture, *LODSettings, *CurrentPlatform, SourceBuildSettings);
 
 		TArray< TArray<FName> > PlatformFormats;
 		CurrentPlatform->GetTextureFormats(&Texture, PlatformFormats);
@@ -1626,10 +1639,7 @@ void UTexture::BeginCacheForCookedPlatformData( const ITargetPlatform *TargetPla
 		//TArray<FName> PlatformFormats;
 
 		FTextureBuildSettings BuildSettings;
-		const bool bPlatformSupportsTextureStreaming = TargetPlatform->SupportsFeature(ETargetPlatformFeatures::TextureStreaming);
-		const bool bPlatformSupportsVirtualTextureStreaming = TargetPlatform->SupportsFeature(ETargetPlatformFeatures::VirtualTextureStreaming);
-
-		GetTextureBuildSettings(*this, TargetPlatform->GetTextureLODSettings(), bPlatformSupportsTextureStreaming, bPlatformSupportsVirtualTextureStreaming, BuildSettings);
+		GetTextureBuildSettings(*this, TargetPlatform->GetTextureLODSettings(), *TargetPlatform, BuildSettings);
 		
 		TArray< TArray<FTextureBuildSettings> > BuildSettingsToCache;
 		GetBuildSettingsPerFormat(*this, BuildSettings, TargetPlatform, BuildSettingsToCache);
@@ -1710,10 +1720,7 @@ void UTexture::ClearCachedCookedPlatformData( const ITargetPlatform* TargetPlatf
 
 		// Retrieve formats to cache for targetplatform.
 		FTextureBuildSettings BuildSettings;
-		const bool bPlatformSupportsTextureStreaming = TargetPlatform->SupportsFeature(ETargetPlatformFeatures::TextureStreaming);
-		const bool bPlatformSupportsVirtualTextureStreaming = TargetPlatform->SupportsFeature(ETargetPlatformFeatures::VirtualTextureStreaming);
-
-		GetTextureBuildSettings(*this, TargetPlatform->GetTextureLODSettings(), bPlatformSupportsTextureStreaming, bPlatformSupportsVirtualTextureStreaming, BuildSettings);
+		GetTextureBuildSettings(*this, TargetPlatform->GetTextureLODSettings(), *TargetPlatform, BuildSettings);
 
 		TArray< TArray<FTextureBuildSettings> > BuildSettingsToCache;
 		GetBuildSettingsPerFormat(*this, BuildSettings, TargetPlatform, BuildSettingsToCache);
@@ -1771,11 +1778,7 @@ bool UTexture::IsCachedCookedPlatformDataLoaded( const ITargetPlatform* TargetPl
 
 	FTextureBuildSettings BuildSettings;
 	TArray<FTexturePlatformData*> PlatformDataToSerialize;
-
-	const bool bPlatformSupportsTextureStreaming = TargetPlatform->SupportsFeature(ETargetPlatformFeatures::TextureStreaming);
-	const bool bPlatformSupportsVirtualTextureStreaming = TargetPlatform->SupportsFeature(ETargetPlatformFeatures::VirtualTextureStreaming);
-
-	GetTextureBuildSettings(*this, TargetPlatform->GetTextureLODSettings(), bPlatformSupportsTextureStreaming, bPlatformSupportsVirtualTextureStreaming, BuildSettings);
+	GetTextureBuildSettings(*this, TargetPlatform->GetTextureLODSettings(), *TargetPlatform, BuildSettings);
 
 	TArray< TArray<FTextureBuildSettings> > BuildSettingsToCache;
 	GetBuildSettingsPerFormat(*this, BuildSettings, TargetPlatform, BuildSettingsToCache);
@@ -1977,10 +1980,7 @@ void UTexture::SerializeCookedPlatformData(FArchive& Ar)
 		if (!Ar.CookingTarget()->IsServerOnly())
 		{
 			FTextureBuildSettings BuildSettings;
-			const bool bPlatformSupportsTextureStreaming = Ar.CookingTarget()->SupportsFeature(ETargetPlatformFeatures::TextureStreaming);
-			const bool bPlatformSupportsVirtualTextureStreaming = Ar.CookingTarget()->SupportsFeature(ETargetPlatformFeatures::VirtualTextureStreaming);
-
-			GetTextureBuildSettings(*this, Ar.CookingTarget()->GetTextureLODSettings(), bPlatformSupportsTextureStreaming, bPlatformSupportsVirtualTextureStreaming, BuildSettings);
+			GetTextureBuildSettings(*this, Ar.CookingTarget()->GetTextureLODSettings(), *Ar.CookingTarget(), BuildSettings);
 
 			TArray<FTexturePlatformData*> PlatformDataToSerialize;
 
diff --git a/Engine/Source/Runtime/Engine/Private/TextureLODSettings.cpp b/Engine/Source/Runtime/Engine/Private/TextureLODSettings.cpp
--- a/Engine/Source/Runtime/Engine/Private/TextureLODSettings.cpp
+++ b/Engine/Source/Runtime/Engine/Private/TextureLODSettings.cpp
@@ -2,7 +2,9 @@
 
 #include "Engine/TextureLODSettings.h"
 #include "Engine/Texture2D.h"
-
+#include "Interfaces/ITargetPlatform.h"
+#include "Interfaces/ITargetPlatformManagerModule.h"
+#include "PlatformInfo.h"
 
 int32 GUITextureLODBias = 0;
 FAutoConsoleVariableRef CVarUITextureLODBias(
@@ -268,6 +270,31 @@ void UTextureLODSettings::GetMipGenSettings(const UTexture& Texture, TextureMipG
 		bOutBorderColorBlack = true;
 	}
 }
+
+void UTextureLODSettings::GetDownscaleOptions(const UTexture& Texture, const ITargetPlatform& CurrentPlatform, float& Downscale, ETextureDownscaleOptions& DownscaleOptions) const
+{
+	float GroupDownscale = FMath::Clamp(TextureLODGroups[Texture.LODGroup].Downscale, 1.0f, 8.0f);
+	ETextureDownscaleOptions GroupDownscaleOptions = TextureLODGroups[Texture.LODGroup].DownscaleOptions;
+	if (GroupDownscaleOptions == ETextureDownscaleOptions::Default)
+	{
+		GroupDownscaleOptions = ETextureDownscaleOptions::SimpleAverage;
+	}
+		
+	const PlatformInfo::FPlatformInfo& Info = CurrentPlatform.GetPlatformInfo();
+	Downscale = Texture.Downscale.GetValueForPlatformIdentifiers(Info.PlatformGroupName, Info.VanillaPlatformName);;
+	if (Downscale < 1.f)
+	{
+		// value is not overriden for this texture
+		Downscale = GroupDownscale;
+	}
+
+	DownscaleOptions = Texture.DownscaleOptions;
+	if (DownscaleOptions == ETextureDownscaleOptions::Default)
+	{
+		DownscaleOptions = GroupDownscaleOptions;
+	}
+}
+
 #endif // #if WITH_EDITORONLY_DATA
 
 
