UE5 - Fix check for mip generation of odd-sized textures
#rnx

diff --git a/Engine/Source/Developer/TextureCompressor/Private/TextureCompressorModule.cpp b/Engine/Source/Developer/TextureCompressor/Private/TextureCompressorModule.cpp
--- a/Engine/Source/Developer/TextureCompressor/Private/TextureCompressorModule.cpp
+++ b/Engine/Source/Developer/TextureCompressor/Private/TextureCompressorModule.cpp
@@ -425,9 +425,16 @@ static void GenerateSharpenedMipB8G8R8A8Templ(
 	bool bSharpenWithoutColorShift,
 	bool bUnfiltered)
 {
-	check( SourceImageData.SizeX == ScaleFactor * DestImageData.SizeX || DestImageData.SizeX == 1 );
-	check( SourceImageData.SizeY == ScaleFactor * DestImageData.SizeY || DestImageData.SizeY == 1 );
-	check( Kernel.GetFilterTableSize() >= 2 );
+	// Add the 0.5 as the source size could be odd
+	checkf( SourceImageData.SizeX == (uint32)((float)ScaleFactor * ((float)DestImageData.SizeX + 0.5f)) || DestImageData.SizeX == 1,
+		TEXT("SourceImage Size %d, %d; ScaleFactor %d; DestImage Size %d, %d"), 
+		SourceImageData.SizeX, SourceImageData.SizeY, ScaleFactor,
+		DestImageData.SizeX, DestImageData.SizeY);
+	checkf( SourceImageData.SizeY == (uint32)((float)ScaleFactor * ((float)DestImageData.SizeY + 0.5f)) || DestImageData.SizeY == 1,
+		TEXT("SourceImage Size %d, %d; ScaleFactor %d; DestImage Size %d, %d"), 
+		SourceImageData.SizeX, SourceImageData.SizeY, ScaleFactor,
+		DestImageData.SizeX, DestImageData.SizeY);
+	checkf( Kernel.GetFilterTableSize() >= 2, TEXT("Kernel table size %d, expected at least 2!"), Kernel.GetFilterTableSize());
 
 	const int32 KernelCenter = (int32)Kernel.GetFilterTableSize() / 2 - 1;
 
