UE5 - Update to 14516198
#rnx

diff --git a/Engine/Source/Developer/TextureCompressor/Private/TextureCompressorModule.cpp b/Engine/Source/Developer/TextureCompressor/Private/TextureCompressorModule.cpp
--- a/Engine/Source/Developer/TextureCompressor/Private/TextureCompressorModule.cpp
+++ b/Engine/Source/Developer/TextureCompressor/Private/TextureCompressorModule.cpp
@@ -426,11 +426,11 @@ static void GenerateSharpenedMipB8G8R8A8Templ(
 	bool bUnfiltered)
 {
 	// Add the 0.5 as the source size could be odd
-	checkf( SourceImageData.SizeX == (uint32)((float)ScaleFactor * ((float)DestImageData.SizeX + 0.5f)) || DestImageData.SizeX == 1,
+	checkf( SourceImageData.SizeX == (uint32)FMath::TruncToInt((float)ScaleFactor * ((float)DestImageData.SizeX + 0.5f)) || DestImageData.SizeX == 1,
 		TEXT("SourceImage Size %d, %d; ScaleFactor %d; DestImage Size %d, %d"), 
 		SourceImageData.SizeX, SourceImageData.SizeY, ScaleFactor,
 		DestImageData.SizeX, DestImageData.SizeY);
-	checkf( SourceImageData.SizeY == (uint32)((float)ScaleFactor * ((float)DestImageData.SizeY + 0.5f)) || DestImageData.SizeY == 1,
+	checkf( SourceImageData.SizeY == (uint32)FMath::TruncToInt((float)ScaleFactor * ((float)DestImageData.SizeY + 0.5f)) || DestImageData.SizeY == 1,
 		TEXT("SourceImage Size %d, %d; ScaleFactor %d; DestImage Size %d, %d"), 
 		SourceImageData.SizeX, SourceImageData.SizeY, ScaleFactor,
 		DestImageData.SizeX, DestImageData.SizeY);
