UE5 - Update to 14516366
#rnx

diff --git a/Engine/Source/Developer/TextureCompressor/Private/TextureCompressorModule.cpp b/Engine/Source/Developer/TextureCompressor/Private/TextureCompressorModule.cpp
--- a/Engine/Source/Developer/TextureCompressor/Private/TextureCompressorModule.cpp
+++ b/Engine/Source/Developer/TextureCompressor/Private/TextureCompressorModule.cpp
@@ -425,15 +425,9 @@ static void GenerateSharpenedMipB8G8R8A8Templ(
 	bool bSharpenWithoutColorShift,
 	bool bUnfiltered)
 {
-	// Add the 0.5 as the source size could be odd
-	checkf( SourceImageData.SizeX == (uint32)FMath::TruncToInt((float)ScaleFactor * ((float)DestImageData.SizeX + 0.5f)) || DestImageData.SizeX == 1,
-		TEXT("SourceImage Size %d, %d; ScaleFactor %d; DestImage Size %d, %d"), 
-		SourceImageData.SizeX, SourceImageData.SizeY, ScaleFactor,
-		DestImageData.SizeX, DestImageData.SizeY);
-	checkf( SourceImageData.SizeY == (uint32)FMath::TruncToInt((float)ScaleFactor * ((float)DestImageData.SizeY + 0.5f)) || DestImageData.SizeY == 1,
-		TEXT("SourceImage Size %d, %d; ScaleFactor %d; DestImage Size %d, %d"), 
-		SourceImageData.SizeX, SourceImageData.SizeY, ScaleFactor,
-		DestImageData.SizeX, DestImageData.SizeY);
+	// These don't work on even texture sizes
+	//check( SourceImageData.SizeX == ScaleFactor * DestImageData.SizeX || DestImageData.SizeX == 1 );
+	//check( SourceImageData.SizeY == ScaleFactor * DestImageData.SizeY || DestImageData.SizeY == 1 );
 	checkf( Kernel.GetFilterTableSize() >= 2, TEXT("Kernel table size %d, expected at least 2!"), Kernel.GetFilterTableSize());
 
 	const int32 KernelCenter = (int32)Kernel.GetFilterTableSize() / 2 - 1;
@@ -604,8 +598,9 @@ static void GenerateMipBorder(
 	FImageView2D& DestImageData
 	)
 {
-	check( SrcImageData.SizeX == 2 * DestImageData.SizeX || DestImageData.SizeX == 1 );
-	check( SrcImageData.SizeY == 2 * DestImageData.SizeY || DestImageData.SizeY == 1 );
+	// These don't work on even texture sizes
+	//check( SrcImageData.SizeX == 2 * DestImageData.SizeX || DestImageData.SizeX == 1 );
+	//check( SrcImageData.SizeY == 2 * DestImageData.SizeY || DestImageData.SizeY == 1 );
 
 	for ( int32 DestY = 0; DestY < DestImageData.SizeY; DestY++ )
 	{
