Fixed texture LODGroup settings not being correctly setup at engine initialization.
Enabled back checks on power of 2 textures when generating mips.
#jyi rolando.caloca

diff --git a/Engine/Source/Developer/TextureCompressor/Private/TextureCompressorModule.cpp b/Engine/Source/Developer/TextureCompressor/Private/TextureCompressorModule.cpp
--- a/Engine/Source/Developer/TextureCompressor/Private/TextureCompressorModule.cpp
+++ b/Engine/Source/Developer/TextureCompressor/Private/TextureCompressorModule.cpp
@@ -425,9 +425,8 @@ static void GenerateSharpenedMipB8G8R8A8Templ(
 	bool bSharpenWithoutColorShift,
 	bool bUnfiltered)
 {
-	// These don't work on even texture sizes
-	//check( SourceImageData.SizeX == ScaleFactor * DestImageData.SizeX || DestImageData.SizeX == 1 );
-	//check( SourceImageData.SizeY == ScaleFactor * DestImageData.SizeY || DestImageData.SizeY == 1 );
+	check( SourceImageData.SizeX == ScaleFactor * DestImageData.SizeX || DestImageData.SizeX == 1 );
+	check( SourceImageData.SizeY == ScaleFactor * DestImageData.SizeY || DestImageData.SizeY == 1 );
 	checkf( Kernel.GetFilterTableSize() >= 2, TEXT("Kernel table size %d, expected at least 2!"), Kernel.GetFilterTableSize());
 
 	const int32 KernelCenter = (int32)Kernel.GetFilterTableSize() / 2 - 1;
@@ -598,9 +597,8 @@ static void GenerateMipBorder(
 	FImageView2D& DestImageData
 	)
 {
-	// These don't work on even texture sizes
-	//check( SrcImageData.SizeX == 2 * DestImageData.SizeX || DestImageData.SizeX == 1 );
-	//check( SrcImageData.SizeY == 2 * DestImageData.SizeY || DestImageData.SizeY == 1 );
+	check( SrcImageData.SizeX == 2 * DestImageData.SizeX || DestImageData.SizeX == 1 );
+	check( SrcImageData.SizeY == 2 * DestImageData.SizeY || DestImageData.SizeY == 1 );
 
 	for ( int32 DestY = 0; DestY < DestImageData.SizeY; DestY++ )
 	{
diff --git a/Engine/Source/Runtime/CoreUObject/Private/UObject/Class.cpp b/Engine/Source/Runtime/CoreUObject/Private/UObject/Class.cpp
--- a/Engine/Source/Runtime/CoreUObject/Private/UObject/Class.cpp
+++ b/Engine/Source/Runtime/CoreUObject/Private/UObject/Class.cpp
@@ -5244,6 +5244,8 @@ void UClass::AssembleReferenceTokenStreams()
 
 const FString UClass::GetConfigName() const
 {
+	static const FName NAME_DeviceProfiles("DeviceProfiles");
+
 	if( ClassConfigName == NAME_Engine )
 	{
 		return GEngineIni;
@@ -5272,6 +5274,10 @@ const FString UClass::GetConfigName() const
 	{
 		return GEditorKeyBindingsIni;
 	}
+	else if (ClassConfigName == NAME_DeviceProfiles)
+	{
+		return GDeviceProfilesIni;
+	}
 	else if( ClassConfigName == NAME_None )
 	{
 		UE_LOG(LogClass, Fatal,TEXT("UObject::GetConfigName() called on class with config name 'None'. Class flags = 0x%08X"), (uint32)ClassFlags );
