U5 - Move TextureCompression into ParallelFor

diff --git a/Engine/Source/Developer/TextureCompressor/Private/TextureCompressorModule.cpp b/Engine/Source/Developer/TextureCompressor/Private/TextureCompressorModule.cpp
--- a/Engine/Source/Developer/TextureCompressor/Private/TextureCompressorModule.cpp
+++ b/Engine/Source/Developer/TextureCompressor/Private/TextureCompressorModule.cpp
@@ -1979,7 +1979,7 @@ void ApplyCompositeTexture(FImage& RoughnessSourceMips, const FImage& NormalSour
 /**
  * Asynchronous compression, used for compressing mips simultaneously.
  */
-class FAsyncCompressionWorker : public FNonAbandonableTask 
+class FAsyncCompressionWorker
 {
 public:
 	/**
@@ -2013,11 +2013,6 @@ public:
 			);
 	}
 
-	FORCEINLINE TStatId GetStatId() const
-	{
-		RETURN_QUICK_DECLARE_CYCLE_STAT(FAsyncCompressionWorker, STATGROUP_ThreadPoolAsyncTasks);
-	}
-
 	/**
 	 * Transfer the result of the compression to the OutCompressedImage
 	 * Can only be called once
@@ -2047,7 +2042,6 @@ private:
 	/** true if compression was successful. */
 	bool bCompressionResults;
 };
-typedef FAsyncTask<FAsyncCompressionWorker> FAsyncCompressionTask;
 
 // compress mip-maps in InMipChain and add mips to Texture, might alter the source content
 static bool CompressMipChain(
@@ -2067,7 +2061,6 @@ static bool CompressMipChain(
 	OutNumMipsInTail = CompressorCaps.NumMipsInTail;
 	OutExtData = CompressorCaps.ExtData;
 
-	TIndirectArray<FAsyncCompressionTask> AsyncCompressionTasks;
 	int32 MipCount = MipChain.Num();
 	check(MipCount >= (int32)CompressorCaps.NumMipsInTail);
 	// This number was too small (128) for current hardware and caused too many
@@ -2085,50 +2078,68 @@ static bool CompressMipChain(
 	}
 
 	OutMips.Empty(MipCount);
+	TArray<FAsyncCompressionWorker> AsyncCompressionTasks;
+	AsyncCompressionTasks.Reserve(MipCount);
+
+	struct PreWork
+	{
+		int32 MipIndex;
+		const FImage& SrcMip;
+		FCompressedImage2D& DestMip;
+	};
+	TArray<PreWork> PreWorkTasks;
+	PreWorkTasks.Reserve(MipCount);
+
 	for (int32 MipIndex = 0; MipIndex < MipCount; ++MipIndex)
 	{
 		const FImage& SrcMip = MipChain[MipIndex];
 		FCompressedImage2D& DestMip = *new(OutMips) FCompressedImage2D;
+
 		if (MipIndex > FirstMipTailIndex)
 		{
 			continue;
 		}
 		else if (bAllowParallelBuild && FMath::Min(SrcMip.SizeX, SrcMip.SizeY) >= MinAsyncCompressionSize)
 		{
-			FAsyncCompressionTask* AsyncTask = new FAsyncCompressionTask(
+			AsyncCompressionTasks.Emplace(
 				TextureFormat,
 				&SrcMip,
 				MipIndex == FirstMipTailIndex ? CompressorCaps.NumMipsInTail : 1, // number of mips pointed to by SrcMip
 				Settings,
 				bImageHasAlphaChannel,
 				CompressorCaps.ExtData
-				);
-			AsyncCompressionTasks.Add(AsyncTask);
-#if WITH_EDITOR
-			AsyncTask->StartBackgroundTask(GLargeThreadPool);
-#else
-			AsyncTask->StartBackgroundTask();
-#endif
+			);
 		}
 		else
+		{
+			PreWorkTasks.Emplace(PreWork { MipIndex, SrcMip, DestMip });
+		}
+	}
+
+	ParallelForWithPreWork(AsyncCompressionTasks.Num(), [&AsyncCompressionTasks](int32 TaskIndex)
+	{
+		AsyncCompressionTasks[TaskIndex].DoWork();
+	},
+	[&PreWorkTasks, &TextureFormat, &OutMips, &bCompressionSucceeded, &CompressorCaps, &Settings, FirstMipTailIndex, bImageHasAlphaChannel]()
+	{
+		for (PreWork& Work : PreWorkTasks)
 		{
 			bCompressionSucceeded = bCompressionSucceeded && TextureFormat->CompressImageEx(
-				&SrcMip,
-				MipIndex == FirstMipTailIndex ? CompressorCaps.NumMipsInTail : 1, // number of mips pointed to by SrcMip
+				&Work.SrcMip,
+				Work.MipIndex == FirstMipTailIndex ? CompressorCaps.NumMipsInTail : 1, // number of mips pointed to by SrcMip
 				Settings,
 				bImageHasAlphaChannel,
 				CompressorCaps.ExtData,
-				DestMip
-				);
+				Work.DestMip
+			);
 		}
-	}
+	}, EParallelForFlags::BackgroundPriority | EParallelForFlags::Unbalanced);
 
 	for (int32 TaskIndex = 0; TaskIndex < AsyncCompressionTasks.Num(); ++TaskIndex)
 	{
-		FAsyncCompressionTask& AsynTask = AsyncCompressionTasks[TaskIndex];
-		AsynTask.EnsureCompletion();
+		FAsyncCompressionWorker& AsynTask = AsyncCompressionTasks[TaskIndex];
 		FCompressedImage2D& DestMip = OutMips[TaskIndex];
-		bCompressionSucceeded = bCompressionSucceeded && AsynTask.GetTask().ConsumeCompressionResults(DestMip);
+		bCompressionSucceeded = bCompressionSucceeded && AsynTask.ConsumeCompressionResults(DestMip);
 	}
 
 	for (int32 MipIndex = FirstMipTailIndex + 1; MipIndex < MipCount; ++MipIndex)
diff --git a/Engine/Source/Developer/TextureFormatDXT/Private/TextureFormatDXT.cpp b/Engine/Source/Developer/TextureFormatDXT/Private/TextureFormatDXT.cpp
--- a/Engine/Source/Developer/TextureFormatDXT/Private/TextureFormatDXT.cpp
+++ b/Engine/Source/Developer/TextureFormatDXT/Private/TextureFormatDXT.cpp
@@ -13,6 +13,8 @@
 #include "TextureCompressorModule.h"
 #include "PixelFormat.h"
 #include "EngineLogs.h"
+#include "Async/ParallelFor.h"
+
 THIRD_PARTY_INCLUDES_START
 	#include "nvtt/nvtt.h"
 THIRD_PARTY_INCLUDES_END
@@ -231,7 +233,7 @@ public:
 /**
  * Asynchronous NVTT worker.
  */
-class FAsyncNVTTWorker : public FNonAbandonableTask 
+class FAsyncNVTTWorker
 {
 public:
 	/**
@@ -249,11 +251,6 @@ public:
 		bCompressionResults = Compressor->Compress();
 	}
 
-	FORCEINLINE TStatId GetStatId() const
-	{
-		RETURN_QUICK_DECLARE_CYCLE_STAT(FAsyncNVTTWorker, STATGROUP_ThreadPoolAsyncTasks);
-	}
-
 	/** Retrieve compression results. */
 	bool GetCompressionResults() const { return bCompressionResults; }
 
@@ -263,7 +260,6 @@ private:
 	/** true if compression was successful. */
 	bool bCompressionResults;
 };
-typedef FAsyncTask<FAsyncNVTTWorker> FAsyncNVTTTask;
 
 namespace CompressionSettings
 {
@@ -374,22 +370,22 @@ static bool CompressImageUsingNVTT(
 	// Asynchronously compress each batch.
 	bool bSuccess = true;
 	{
-		TIndirectArray<FAsyncNVTTTask> AsyncTasks;
+		TArray<FAsyncNVTTWorker> AsyncTasks;
+		AsyncTasks.Reserve(NumBatches);
+
 		for (int32 BatchIndex = 0; BatchIndex < NumBatches; ++BatchIndex)
 		{
-			FAsyncNVTTTask* AsyncTask = new FAsyncNVTTTask(&Compressors[BatchIndex]);
-			AsyncTasks.Add(AsyncTask);
-#if WITH_EDITOR
-			AsyncTask->StartBackgroundTask(GLargeThreadPool);
-#else
-			AsyncTask->StartBackgroundTask();
-#endif
+			AsyncTasks.Emplace(&Compressors[BatchIndex]);
 		}
+
+		ParallelForTemplate(AsyncTasks.Num(), [&AsyncTasks](int32 TaskIndex)
+		{
+			AsyncTasks[TaskIndex].DoWork();
+		}, EParallelForFlags::BackgroundPriority | EParallelForFlags::Unbalanced);
+
 		for (int32 BatchIndex = 0; BatchIndex < NumBatches; ++BatchIndex)
 		{
-			FAsyncNVTTTask& AsyncTask = AsyncTasks[BatchIndex];
-			AsyncTask.EnsureCompletion();
-			bSuccess = bSuccess && AsyncTask.GetTask().GetCompressionResults();
+			bSuccess = bSuccess && AsyncTasks[BatchIndex].GetCompressionResults();
 		}
 	}
 
diff --git a/Engine/Source/Developer/TextureFormatIntelISPCTexComp/Private/TextureFormatIntelISPCTexComp.cpp b/Engine/Source/Developer/TextureFormatIntelISPCTexComp/Private/TextureFormatIntelISPCTexComp.cpp
--- a/Engine/Source/Developer/TextureFormatIntelISPCTexComp/Private/TextureFormatIntelISPCTexComp.cpp
+++ b/Engine/Source/Developer/TextureFormatIntelISPCTexComp/Private/TextureFormatIntelISPCTexComp.cpp
@@ -13,6 +13,7 @@
 #include "Interfaces/ITextureFormatModule.h"
 #include "TextureCompressorModule.h"
 #include "PixelFormat.h"
+#include "Async/ParallelFor.h"
 
 #include "ispc_texcomp.h"
 
@@ -237,7 +238,7 @@ struct FMultithreadedCompression
 	{
 		if (bUseTasks)
 		{
-			class FIntelCompressWorker : public FNonAbandonableTask
+			class FIntelCompressWorker
 			{
 			public:
 				FIntelCompressWorker(EncoderSettingsType* pEncSettings, FImage* pInImage, FCompressedImage2D* pOutImage, int yStart, int yEnd, int SliceIndex, CompressFunction InFunctionCallback)
@@ -256,11 +257,6 @@ struct FMultithreadedCompression
 					mCallback(mpEncSettings, mpInImage, mpOutImage, mYStart, mYEnd, mSliceIndex);
 				}
 
-				FORCEINLINE TStatId GetStatId() const
-				{
-					RETURN_QUICK_DECLARE_CYCLE_STAT(FIntelCompressWorker, STATGROUP_ThreadPoolAsyncTasks);
-				}
-
 				EncoderSettingsType*	mpEncSettings;
 				FImage*					mpInImage;
 				FCompressedImage2D*		mpOutImage;
@@ -269,10 +265,9 @@ struct FMultithreadedCompression
 				int						mSliceIndex;
 				CompressFunction		mCallback;
 			};
-			typedef FAsyncTask<FIntelCompressWorker> FIntelCompressTask;
 
 			// One less task because we'll do the final + non multiple of 4 inside this task
-			TIndirectArray<FIntelCompressTask> CompressionTasks;
+			TArray<FIntelCompressWorker> CompressionTasks;
 			const int NumStasksPerSlice = MultithreadSettings.iNumTasks + 1;
 			CompressionTasks.Reserve(NumStasksPerSlice * Image.NumSlices - 1);
 			for (int SliceIndex = 0; SliceIndex < Image.NumSlices; ++SliceIndex)
@@ -282,19 +277,19 @@ struct FMultithreadedCompression
 					// Create a new task unless it's the last task in the last slice (that one will run on current thread, after these threads have been started)
 					if (SliceIndex < (Image.NumSlices - 1) || iTask < (NumStasksPerSlice - 1))
 					{
-						auto* AsyncTask = new FIntelCompressTask(&EncoderSettings, &Image, &OutCompressedImage, iTask * MultithreadSettings.iScansPerTask, (iTask + 1) * MultithreadSettings.iScansPerTask, SliceIndex, FunctionCallback);
-						CompressionTasks.Add(AsyncTask);
-						AsyncTask->StartBackgroundTask();
+						CompressionTasks.Emplace(&EncoderSettings, &Image, &OutCompressedImage, iTask * MultithreadSettings.iScansPerTask, (iTask + 1) * MultithreadSettings.iScansPerTask, SliceIndex, FunctionCallback);
 					}
 				}
 			}
-			FunctionCallback(&EncoderSettings, &Image, &OutCompressedImage, MultithreadSettings.iScansPerTask * MultithreadSettings.iNumTasks, Image.SizeY, Image.NumSlices - 1);
 
-			// Wait for all tasks to complete
-			for (int32 TaskIndex = 0; TaskIndex < CompressionTasks.Num(); ++TaskIndex)
+			ParallelForWithPreWork(CompressionTasks.Num(), [&CompressionTasks](int32 TaskIndex)
 			{
-				CompressionTasks[TaskIndex].EnsureCompletion();
-			}
+				CompressionTasks[TaskIndex].DoWork();
+			},
+			[&EncoderSettings, &Image, &OutCompressedImage, &MultithreadSettings, &FunctionCallback]()
+			{
+				FunctionCallback(&EncoderSettings, &Image, &OutCompressedImage, MultithreadSettings.iScansPerTask * MultithreadSettings.iNumTasks, Image.SizeY, Image.NumSlices - 1);
+			}, EParallelForFlags::BackgroundPriority | EParallelForFlags::Unbalanced);
 		}
 		else
 		{
