Initialize TextureFormatManager in FTextureCompressorModule::StartupModule. Fixes crash on startup due to TFM singleton being created from a task thread.

diff --git a/Engine/Source/Developer/TextureCompressor/Private/TextureCompressorModule.cpp b/Engine/Source/Developer/TextureCompressor/Private/TextureCompressorModule.cpp
--- a/Engine/Source/Developer/TextureCompressor/Private/TextureCompressorModule.cpp
+++ b/Engine/Source/Developer/TextureCompressor/Private/TextureCompressorModule.cpp
@@ -2202,7 +2202,6 @@ public:
 	virtual bool UsesTaskGraph(const FTextureBuildSettings& BuildSettings) const override
 	{
 		const ITextureFormat* TextureFormat = nullptr;
-		ITextureFormatManagerModule* TFM = GetTextureFormatManager();
 		if (TFM)
 		{
 			TextureFormat = TFM->FindTextureFormat(BuildSettings.TextureFormatName);
@@ -2226,7 +2225,7 @@ public:
 		)
 	{
 		const ITextureFormat* TextureFormat = nullptr;
-		ITextureFormatManagerModule* TFM = GetTextureFormatManager();
+
 		if (TFM)
 		{
 			TextureFormat = TFM->FindTextureFormat(BuildSettings.TextureFormatName);
@@ -2305,6 +2304,9 @@ public:
 		nvTextureToolsHandle = FPlatformProcess::GetDllHandle(*(FPaths::EngineDir() / TEXT("Binaries/ThirdParty/nvTextureTools/Win32/nvtt_.dll")));
 	#endif
 #endif	//PLATFORM_WINDOWS
+
+		// Initialize TFM singleton on the main thread, as it may be accessed from task threads later.
+		TFM = GetTextureFormatManager();
 	}
 
 	void ShutdownModule()
@@ -2321,6 +2323,8 @@ private:
 	void* nvTextureToolsHandle;
 #endif	//PLATFORM_WINDOWS
 
+	ITextureFormatManagerModule* TFM = nullptr;
+
 	bool BuildTextureMips(
 		const TArray<FImage>& InSourceMips,
 		const FTextureBuildSettings& BuildSettings,
