Disable deprecated cubemap processing originally designed for ambient lighting.

diff --git a/Engine/Source/Developer/TextureCompressor/Private/TextureCompressorModule.cpp b/Engine/Source/Developer/TextureCompressor/Private/TextureCompressorModule.cpp
--- a/Engine/Source/Developer/TextureCompressor/Private/TextureCompressorModule.cpp
+++ b/Engine/Source/Developer/TextureCompressor/Private/TextureCompressorModule.cpp
@@ -1201,9 +1201,9 @@ static FVector ComputeWSCubeDirectionAtTexelCenter(uint32 CubemapFace, uint32 x,
 	return DirectionWS;
 }
 
-static int32 ComputeLongLatCubemapExtents(const FImage& SrcImage, const int32 MaxCubemapTextureResolution)
+static uint32 ComputeLongLatCubemapExtents(const FImage& SrcImage, const uint32 MaxCubemapTextureResolution)
 {
-	return FMath::Clamp(1 << FMath::FloorLog2(SrcImage.SizeX / 2), 32, MaxCubemapTextureResolution);
+	return FMath::Clamp(1U << FMath::FloorLog2(SrcImage.SizeX / 2), 32U, MaxCubemapTextureResolution);
 }
 
 void ITextureCompressorModule::GenerateBaseCubeMipFromLongitudeLatitude2D(FImage* OutMip, const FImage& SrcImage, const int32 MaxCubemapTextureResolution)
@@ -1213,16 +1213,16 @@ void ITextureCompressorModule::GenerateBaseCubeMipFromLongitudeLatitude2D(FImage
 	FImageViewLongLat LongLatView(LongLatImage);
 
 	// TODO_TEXTURE: Expose target size to user.
-	int32 Extent = ComputeLongLatCubemapExtents(LongLatImage, MaxCubemapTextureResolution);
+	uint32 Extent = ComputeLongLatCubemapExtents(LongLatImage, MaxCubemapTextureResolution);
 	float InvExtent = 1.0f / Extent;
 	OutMip->Init(Extent, Extent, 6, ERawImageFormat::RGBA32F, EGammaSpace::Linear);
 
 	for(uint32 Face = 0; Face < 6; ++Face)
 	{
 		FImageView2D MipView(*OutMip, Face);
-		for(int32 y = 0; y < Extent; ++y)
+		for(uint32 y = 0; y < Extent; ++y)
 		{
-			for(int32 x = 0; x < Extent; ++x)
+			for(uint32 x = 0; x < Extent; ++x)
 			{
 				FVector DirectionWS = ComputeWSCubeDirectionAtTexelCenter(Face, x, y, InvExtent);
 				MipView.Access(x, y) = LongLatView.LookupLongLat(DirectionWS);
@@ -2585,15 +2585,7 @@ private:
 		// Generate any missing mips in the chain.
 		if (NumOutputMips > OutMipChain.Num())
 		{
-			// Do angular filtering of cubemaps if requested.
-			if (BuildSettings.bCubemap)
-			{
-				GenerateAngularFilteredMips(OutMipChain, NumOutputMips, BuildSettings.DiffuseConvolveMipLevel);
-			}
-			else
-			{
-				GenerateMipChain(BuildSettings, OutMipChain.Last(), OutMipChain);
-			}
+			GenerateMipChain(BuildSettings, OutMipChain.Last(), OutMipChain);
 		}
 		check(OutMipChain.Num() == NumOutputMips);
 
diff --git a/Engine/Source/Runtime/Engine/Private/TextureDerivedData.cpp b/Engine/Source/Runtime/Engine/Private/TextureDerivedData.cpp
--- a/Engine/Source/Runtime/Engine/Private/TextureDerivedData.cpp
+++ b/Engine/Source/Runtime/Engine/Private/TextureDerivedData.cpp
@@ -50,7 +50,7 @@
 // In case of merge conflicts with DDC versions, you MUST generate a new GUID and set this new
 // guid as version
 
-#define TEXTURE_DERIVEDDATA_VER		TEXT("0C50BD8A80BA468C94A214E8B51B3B4D")
+#define TEXTURE_DERIVEDDATA_VER		TEXT("E405EE6911174BB2A4B24BAC98C41448")
 
 // This GUID is mixed into DDC version for virtual textures only, this allows updating DDC version for VT without invalidating DDC for all textures
 // This is useful during development, but once large numbers of VT are present in shipped content, it will have the same problem as TEXTURE_DERIVEDDATA_VER
@@ -396,11 +396,6 @@ static void GetTextureBuildSettings(
 		OutBuildSettings.DiffuseConvolveMipLevel = GDiffuseConvolveMipLevel;
 		const UTextureCube* Cube = CastChecked<UTextureCube>(&Texture);
 		OutBuildSettings.bLongLatSource = (Cube->Source.GetNumSlices() == 1);
-		if (OutBuildSettings.bLongLatSource && Texture.MaxTextureSize <= 0)
-		{
-			// long/lat source use 512 as default
-			OutBuildSettings.MaxTextureResolution = 512;
-		}
 	}
 	else if (Texture.IsA(UTexture2DArray::StaticClass()))
 	{
