Use proper argument types in ITextureCompressorModule::GenerateBaseCubeMipFromLongitudeLatitude2D.

diff --git a/Engine/Source/Developer/TextureCompressor/Private/TextureCompressorModule.cpp b/Engine/Source/Developer/TextureCompressor/Private/TextureCompressorModule.cpp
--- a/Engine/Source/Developer/TextureCompressor/Private/TextureCompressorModule.cpp
+++ b/Engine/Source/Developer/TextureCompressor/Private/TextureCompressorModule.cpp
@@ -1206,7 +1206,7 @@ static uint32 ComputeLongLatCubemapExtents(const FImage& SrcImage, const uint32
 	return FMath::Clamp(1U << FMath::FloorLog2(SrcImage.SizeX / 2), 32U, MaxCubemapTextureResolution);
 }
 
-void ITextureCompressorModule::GenerateBaseCubeMipFromLongitudeLatitude2D(FImage* OutMip, const FImage& SrcImage, const int32 MaxCubemapTextureResolution)
+void ITextureCompressorModule::GenerateBaseCubeMipFromLongitudeLatitude2D(FImage* OutMip, const FImage& SrcImage, const uint32 MaxCubemapTextureResolution)
 {
 	FImage LongLatImage;
 	SrcImage.CopyTo(LongLatImage, ERawImageFormat::RGBA32F, EGammaSpace::Linear);
diff --git a/Engine/Source/Developer/TextureCompressor/Public/TextureCompressorModule.h b/Engine/Source/Developer/TextureCompressor/Public/TextureCompressorModule.h
--- a/Engine/Source/Developer/TextureCompressor/Public/TextureCompressorModule.h
+++ b/Engine/Source/Developer/TextureCompressor/Public/TextureCompressorModule.h
@@ -302,7 +302,7 @@ public:
 	 * @param OutMip - The output mip.
 	 * @param SrcImage - The source longlat image.
 	 */
-	TEXTURECOMPRESSOR_API static void GenerateBaseCubeMipFromLongitudeLatitude2D(FImage* OutMip, const FImage& SrcImage, const int32 MaxCubemapTextureResolution);
+	TEXTURECOMPRESSOR_API static void GenerateBaseCubeMipFromLongitudeLatitude2D(FImage* OutMip, const FImage& SrcImage, const uint32 MaxCubemapTextureResolution);
 
 
 	/**
