Avoid a deadlock when running Oodle texture compression on low core count machines.  Suggested by andriy.tylychko in this JIRA:
https://jira.it.epicgames.com/browse/UE-119594
as a temporary measure until arne.schober has a chance to put in place a long term fix.

diff --git a/Engine/Source/Developer/TextureCompressor/Private/TextureCompressorModule.cpp b/Engine/Source/Developer/TextureCompressor/Private/TextureCompressorModule.cpp
--- a/Engine/Source/Developer/TextureCompressor/Private/TextureCompressorModule.cpp
+++ b/Engine/Source/Developer/TextureCompressor/Private/TextureCompressorModule.cpp
@@ -2133,7 +2133,7 @@ static bool CompressMipChain(
 				Work.DestMip
 			);
 		}
-	}, EParallelForFlags::Unbalanced);
+	}, EParallelForFlags::Unbalanced | EParallelForFlags::BackgroundPriority);
 
 	for (int32 TaskIndex = 0; TaskIndex < AsyncCompressionTasks.Num(); ++TaskIndex)
 	{
