Undo CL16853912 now that CL16922241 fixes the deadlock properly.

diff --git a/Engine/Source/Developer/TextureCompressor/Private/TextureCompressorModule.cpp b/Engine/Source/Developer/TextureCompressor/Private/TextureCompressorModule.cpp
--- a/Engine/Source/Developer/TextureCompressor/Private/TextureCompressorModule.cpp
+++ b/Engine/Source/Developer/TextureCompressor/Private/TextureCompressorModule.cpp
@@ -2133,7 +2133,7 @@ static bool CompressMipChain(
 				Work.DestMip
 			);
 		}
-	}, EParallelForFlags::Unbalanced | EParallelForFlags::BackgroundPriority);
+	}, EParallelForFlags::Unbalanced);
 
 	for (int32 TaskIndex = 0; TaskIndex < AsyncCompressionTasks.Num(); ++TaskIndex)
 	{
