cleanups in ComputeAlphaCoverage , output value NOP
verified computes the exact same value as before
ComputeAlphaCoverage does not need to be templated because it does not reference pixels off edge of the mip

diff --git a/Engine/Source/Developer/TextureCompressor/Private/TextureCompressorModule.cpp b/Engine/Source/Developer/TextureCompressor/Private/TextureCompressorModule.cpp
--- a/Engine/Source/Developer/TextureCompressor/Private/TextureCompressorModule.cpp
+++ b/Engine/Source/Developer/TextureCompressor/Private/TextureCompressorModule.cpp
@@ -313,7 +313,6 @@ private:
 	float KernelWeights[MaxKernelExtend * MaxKernelExtend];
 };
 
-template <EMipGenAddressMode AddressMode>
 static FVector4f ComputeAlphaCoverage(const FVector4f& Thresholds, const FVector4f& Scales, const FImageView2D& SourceImageData)
 {
 	TRACE_CPUPROFILER_EVENT_SCOPE(ComputeAlphaCoverage);
@@ -323,69 +322,131 @@ static FVector4f ComputeAlphaCoverage(const FVector4f& Thresholds, const FVector
 	int32 NumRowsEachJob;
 	int32 NumJobs = ImageParallelForComputeNumJobsForRows(NumRowsEachJob,SourceImageData.SizeX,SourceImageData.SizeY);
 
-	int32 CommonResults[4] = { 0, 0, 0, 0 };
-	ParallelFor(NumJobs, [&](int32 Index)
+	if ( Thresholds[0] == 0.f && Thresholds[1] == 0.f && Thresholds[2] == 0.f )
 	{
-		TRACE_CPUPROFILER_EVENT_SCOPE(ComputeAlphaCoverage.PF);
+		// common case that only channel 3 (A) is used for alpha coverage :
+		
+		check( Thresholds[3] != 0.f );
 
-		int32 StartIndex = Index * NumRowsEachJob;
-		int32 EndIndex = FMath::Min(StartIndex + NumRowsEachJob, SourceImageData.SizeY);
-		int32 LocalCoverage[4] = { 0, 0, 0, 0 };
-		for (int32 y = StartIndex; y < EndIndex; ++y)
+		const float ThresholdScaled = Thresholds[3] / Scales[3];
+		
+		int32 CommonResult = 0;
+		ParallelFor(NumJobs, [&](int32 Index)
 		{
-			for (int32 x = 0; x < SourceImageData.SizeX; ++x)
+			TRACE_CPUPROFILER_EVENT_SCOPE(ComputeAlphaCoverage.PF);
+
+			int32 StartIndex = Index * NumRowsEachJob;
+			int32 EndIndex = FMath::Min(StartIndex + NumRowsEachJob, SourceImageData.SizeY);
+			int32 LocalCoverage = 0;
+			for (int32 y = StartIndex; y < EndIndex; ++y)
 			{
-				// Sample channel values at pixel neighborhood
-				FVector4f PixelValue(LookupSourceMip<AddressMode>(SourceImageData, x, y));
+				const FLinearColor * RowPixels = &SourceImageData.Access(0,y);
 
-				// Calculate coverage for each channel (if being used as an alpha mask)
-				for (int32 i = 0; i < 4; ++i)
+				for (int32 x = 0; x < SourceImageData.SizeX; ++x)
 				{
-					// Skip channel if Threshold is 0
-					if (Thresholds[i] == 0)
-					{
-						continue;
-					}
+					// Sample channel values at pixel neighborhood
+					const float A = RowPixels[x].A;
+
+					LocalCoverage += (A >= ThresholdScaled);
+				}
+			}
+
+			FPlatformAtomics::InterlockedAdd(&CommonResult, LocalCoverage);
+		});
+
+		Coverage[3] = float(CommonResult) / float(SourceImageData.SizeX * SourceImageData.SizeY);
+		
+		UE_LOG(LogTextureCompressor, VeryVerbose, TEXT("Thresholds = 000 %f Coverage = 000 %f"),  \
+			Thresholds[3], Coverage[3] );
+	}
+	else
+	{
+		FVector4f ThresholdsScaled;
+
+		for (int32 i = 0; i < 4; ++i)
+		{
+			// Skip channel if Threshold is 0
+			if (Thresholds[i] == 0)
+			{
+				// stuff a value that we will always be less than
+				ThresholdsScaled[i] = FLT_MAX;
+			}
+			else
+			{
+				check( Scales[i] != 0.f );
+				ThresholdsScaled[i] = Thresholds[i] / Scales[i];
+			}
+		}
+
+		int32 CommonResults[4] = { 0, 0, 0, 0 };
+		ParallelFor(NumJobs, [&](int32 Index)
+		{
+			TRACE_CPUPROFILER_EVENT_SCOPE(ComputeAlphaCoverage.PF);
 
-					if (PixelValue[i] * Scales[i] >= Thresholds[i])
+			int32 StartIndex = Index * NumRowsEachJob;
+			int32 EndIndex = FMath::Min(StartIndex + NumRowsEachJob, SourceImageData.SizeY);
+			int32 LocalCoverage[4] = { 0, 0, 0, 0 };
+			for (int32 y = StartIndex; y < EndIndex; ++y)
+			{
+				const FLinearColor * RowPixels = &SourceImageData.Access(0,y);
+
+				for (int32 x = 0; x < SourceImageData.SizeX; ++x)
+				{
+					const FLinearColor & PixelValue = RowPixels[x];
+
+					// Calculate coverage for each channel
+					for (int32 i = 0; i < 4; ++i)
 					{
-						++LocalCoverage[i];
+						LocalCoverage[i] += ( PixelValue.Component(i) >= ThresholdsScaled[i] );
 					}
 				}
 			}
-		}
+
+			for (int32 i = 0; i < 4; ++i)
+			{
+				FPlatformAtomics::InterlockedAdd(&CommonResults[i], LocalCoverage[i]);
+			}
+		});
 
 		for (int32 i = 0; i < 4; ++i)
 		{
-			FPlatformAtomics::InterlockedAdd(&CommonResults[i], LocalCoverage[i]);
+			Coverage[i] = float(CommonResults[i]) / float(SourceImageData.SizeX * SourceImageData.SizeY);
 		}
-	});
-
-	for (int32 i = 0; i < 4; ++i)
-	{
-		Coverage[i] = float(CommonResults[i]);
+		
+		UE_LOG(LogTextureCompressor, VeryVerbose, TEXT("Thresholds = %f %f %f %f Coverage = %f %f %f %f"),  \
+			Thresholds[0], Thresholds[1], Thresholds[2], Thresholds[3], \
+			Coverage[0], Coverage[1], Coverage[2], Coverage[3] );
 	}
-
-	return Coverage / float(SourceImageData.SizeX * SourceImageData.SizeY);
+	
+	return Coverage;
 }
 
-template <EMipGenAddressMode AddressMode>
 static FVector4f ComputeAlphaScale(const FVector4f& Coverages, const FVector4f& AlphaThresholds, const FImageView2D& SourceImageData)
 {
 	TRACE_CPUPROFILER_EVENT_SCOPE(ComputeAlphaScale);
 
+	// This function is not a good way to do this
+	// but we cannot change it without changing output pixels
+	// A better method would be to histogram the channel and scale the histogram to meet the desired threshold
+	// even if using this binary search method, you should remember which value gave the closest result
+	//	 don't assume that each binary search step is an improvement
+	// 
+
 	FVector4f MinAlphaScales (0, 0, 0, 0);
 	FVector4f MaxAlphaScales (4, 4, 4, 4);
 	FVector4f AlphaScales (1, 1, 1, 1);
 
 	//Binary Search to find Alpha Scale
+	// limit binary search to 8 steps
 	for (int32 i = 0; i < 8; ++i)
 	{
-		FVector4f ComputedCoverages = ComputeAlphaCoverage<AddressMode>(AlphaThresholds, AlphaScales, SourceImageData);
+		FVector4f ComputedCoverages = ComputeAlphaCoverage(AlphaThresholds, AlphaScales, SourceImageData);
+		
+		UE_LOG(LogTextureCompressor, VeryVerbose, TEXT("Tried AlphaScale = %f ComputedCoverage = %f Goal = %f"), AlphaScales[3], ComputedCoverages[3], Coverages[3] );
 
 		for (int32 j = 0; j < 4; ++j)
 		{
-			if (AlphaThresholds[j] == 0 || fabs(ComputedCoverages[j] - Coverages[j]) < KINDA_SMALL_NUMBER)
+			if (AlphaThresholds[j] == 0 || fabsf(ComputedCoverages[j] - Coverages[j]) < KINDA_SMALL_NUMBER)
 			{
 				continue;
 			}
@@ -399,15 +460,20 @@ static FVector4f ComputeAlphaScale(const FVector4f& Coverages, const FVector4f&
 				MaxAlphaScales[j] = AlphaScales[j];
 			}
 
+			// guess alphascale is best at next midpoint :
+			//  this means we wind up returning an alphascale value we have never tested
 			AlphaScales[j] = (MinAlphaScales[j] + MaxAlphaScales[j]) * 0.5f;
 		}
 
+		// Equals default tolerance is KINDA_SMALL_NUMBER so it checks the same condition as above
 		if (ComputedCoverages.Equals(Coverages))
 		{
 			break;
 		}
 	}
 
+	UE_LOG(LogTextureCompressor, VeryVerbose, TEXT("Final AlphaScales = %f %f %f %f"), AlphaScales[0], AlphaScales[1], AlphaScales[2], AlphaScales[3] );
+
 	return AlphaScales;
 }
 
@@ -445,7 +511,7 @@ static void GenerateSharpenedMipB8G8R8A8Templ(
 	FVector4f AlphaScale(1, 1, 1, 1);
 	if (AlphaThresholds != FVector4f(0,0,0,0))
 	{
-		AlphaScale = ComputeAlphaScale<AddressMode>(AlphaCoverages, AlphaThresholds, SourceImageData);
+		AlphaScale = ComputeAlphaScale(AlphaCoverages, AlphaThresholds, SourceImageData);
 	}
 	
 	int32 NumRowsEachJob;
@@ -966,29 +1032,21 @@ void ITextureCompressorModule::GenerateMipChain(
 	{
 		check(IntermediateSrcPtr);
 		const FImageView2D IntermediateSrcView = FImageView2D::ConstructConst(*IntermediateSrcPtr, 0);
-		switch (AddressMode)
-		{
-		case MGTAM_Wrap:
-			AlphaCoverages = ComputeAlphaCoverage<MGTAM_Wrap>(Settings.AlphaCoverageThresholds, AlphaScales, IntermediateSrcView);
-			break;
-		case MGTAM_Clamp:
-			AlphaCoverages = ComputeAlphaCoverage<MGTAM_Clamp>(Settings.AlphaCoverageThresholds, AlphaScales, IntermediateSrcView);
-			break;
-		case MGTAM_BorderBlack:
-			AlphaCoverages = ComputeAlphaCoverage<MGTAM_BorderBlack>(Settings.AlphaCoverageThresholds, AlphaScales, IntermediateSrcView);
-			break;
-		default:
-			check(0);
-		}		
+		
+		AlphaCoverages = ComputeAlphaCoverage(Settings.AlphaCoverageThresholds, AlphaScales, IntermediateSrcView);
 	}
 
 	// Generate mips
+	//  default value of MipChainDepth is MAX_uint32, means generate all mips down to 1x1
+	//	(break inside the loop)
 	for (; MipChainDepth != 0 ; --MipChainDepth)
 	{
 		check(IntermediateSrcPtr && IntermediateDstPtr);
 		const FImage& IntermediateSrc = *IntermediateSrcPtr;
 		FImage& IntermediateDst = *IntermediateDstPtr;
 
+		// add new mip to TArray<FImage> &OutMipChain :
+		//	placement new on TArray does AddUninitialized then constructs in the last element
 		FImage& DestImage = *new(OutMipChain) FImage(IntermediateDst.SizeX, IntermediateDst.SizeY, IntermediateDst.NumSlices, ImageFormat);
 		
 		for (int32 SliceIndex = 0; SliceIndex < IntermediateDst.NumSlices; ++SliceIndex)
@@ -1640,7 +1698,7 @@ void ITextureCompressorModule::AdjustImageColors(FImage& Image, const FTextureBu
 		//     and thus overhead of multithreading will simply make it slower
 		bool bForceSingleThread = GIsEditorLoadingPackage || GIsCookerLoadingPackage || IsInAsyncLoadingThread();
 
-		TFunction<void (int32)> AdjustImageColorsFunc = [&](int32 Index)
+		TFunctionRef<void (int32)> AdjustImageColorsFunc = [&](int32 Index)
 		{
 			TRACE_CPUPROFILER_EVENT_SCOPE(AdjustImageColors.PF);
 
