fix previous check in TextureCompressorModule
TFunctionRef AdjustImageColorsFunc is not okay, put back to TFunction

diff --git a/Engine/Source/Developer/TextureCompressor/Private/TextureCompressorModule.cpp b/Engine/Source/Developer/TextureCompressor/Private/TextureCompressorModule.cpp
--- a/Engine/Source/Developer/TextureCompressor/Private/TextureCompressorModule.cpp
+++ b/Engine/Source/Developer/TextureCompressor/Private/TextureCompressorModule.cpp
@@ -1698,7 +1698,9 @@ void ITextureCompressorModule::AdjustImageColors(FImage& Image, const FTextureBu
 		//     and thus overhead of multithreading will simply make it slower
 		bool bForceSingleThread = GIsEditorLoadingPackage || GIsCookerLoadingPackage || IsInAsyncLoadingThread();
 
-		TFunctionRef<void (int32)> AdjustImageColorsFunc = [&](int32 Index)
+		// TFunction or auto are okay here
+		// TFunctionRef is not
+		TFunction<void (int32)> AdjustImageColorsFunc = [&](int32 Index)
 		{
 			TRACE_CPUPROFILER_EVENT_SCOPE(AdjustImageColors.PF);
 
