clean up the mess caused by AlphaCoverageThresholds
AlphaCoverageThresholds was accidentally enabled for all textures with thresholds = {0,0,0,1}
add new Texture bool bDoScaleMipsForAlphaCoverage
to make toggling coverage processing very clear
on old Textures that don't have the bDoScaleMipsForAlphaCoverage field, we can infer the correct value except in the case that threshold = {0,0,0,1} , in that case it's impossible to tell if it's wanted or not.  There we look up a config value "EnableLegacyAlphaCoverageThresholdScaling"
Current config has EnableLegacyAlphaCoverageThresholdScaling=1
so the legacy bad behavior is maintained

diff --git a/Engine/Config/BaseEditor.ini b/Engine/Config/BaseEditor.ini
--- a/Engine/Config/BaseEditor.ini
+++ b/Engine/Config/BaseEditor.ini
@@ -410,6 +410,12 @@ Texture2D=50
 [TextureBuild]
 NewTextureBuilds=false
 
+[Texture]
+; EnableLegacyAlphaCoverageThresholdScaling controls whether legacy textures with alphathresholds incorrectly set to {0,0,0,1}
+;   load in with bDoScaleMipsForAlphaCoverage enabled or not
+;   to preserve the same bad alpha processing set EnableLegacyAlphaCoverageThresholdScaling=true
+EnableLegacyAlphaCoverageThresholdScaling=true
+
 [/Script/Localization.LocalizationSettings]
 +EngineTargetsSettings=(Name="Engine",Guid=33482D004789784C9DA695A682ACCA1B,TargetDependencies=,AdditionalManifestDependencies=,RequiredModuleNames=,GatherFromTextFiles=(IsEnabled=True,SearchDirectories=((Path="Source/Runtime/"),(Path="Source/Developer/"),(Path="Config/")),ExcludePathWildcards=((Pattern="Source/Developer/NoRedist/CommunityPortalServices/*"),(Pattern="Source/Developer/NoRedist/UnrealEngineLauncherServices/*"),(Pattern="Source/Developer/NoRedist/BuildPatchServices/*")),FileExtensions=((Pattern="cpp"),(Pattern="h"),(Pattern="c"),(Pattern="inl"),(Pattern="mm"),(Pattern="ini"))),GatherFromPackages=(IsEnabled=True,IncludePathWildcards=((Pattern="Content/*")),ExcludePathWildcards=((Pattern="Content/Editor/*"),(Pattern="Content/Tutorial/*"),(Pattern="Content/Developers/*"),(Pattern="Content/TestPackages/*"),(Pattern="Content/QA_Assets/*"),(Pattern="Content/Maps/Automation/*"),(Pattern="Content/EngineSounds/*")),FileExtensions=((Pattern="umap"),(Pattern="uasset")),ShouldGatherFromEditorOnlyData=True),GatherFromMetaData=(IsEnabled=False,IncludePathWildcards=,ExcludePathWildcards=,KeySpecifications=,ShouldGatherFromEditorOnlyData=True),NativeCultureIndex=0,SupportedCulturesStatistics=((CultureName="en"),(CultureName="es"),(CultureName="ja"),(CultureName="ko"),(CultureName="pt-BR"),(CultureName="zh-CN")))
 +EngineTargetsSettings=(Name="Editor",Guid=AC8BFD2A41A2FB2893BB8EA0AF903E6D,TargetDependencies=(33482D004789784C9DA695A682ACCA1B),AdditionalManifestDependencies=,RequiredModuleNames=,GatherFromTextFiles=(IsEnabled=True,SearchDirectories=((Path="Source/Editor/")),ExcludePathWildcards=,FileExtensions=((Pattern="cpp"),(Pattern="h"),(Pattern="c"),(Pattern="inl"),(Pattern="mm"))),GatherFromPackages=(IsEnabled=True,IncludePathWildcards=((Pattern="Content/Editor/*"),(Pattern="Content/Editor*")),ExcludePathWildcards=,FileExtensions=((Pattern="umap"),(Pattern="uasset")),ShouldGatherFromEditorOnlyData=True),GatherFromMetaData=(IsEnabled=False,IncludePathWildcards=,ExcludePathWildcards=,KeySpecifications=,ShouldGatherFromEditorOnlyData=True),NativeCultureIndex=0,SupportedCulturesStatistics=((CultureName="en"),(CultureName="es"),(CultureName="ja"),(CultureName="ko"),(CultureName="pt-BR"),(CultureName="zh-CN")))
@@ -543,4 +549,4 @@ bCanModifyPluginsFromBrowser=true
 
 [/Script/GameFeaturesEditor.GameFeaturesEditorSettings]
 +PluginTemplates=(Path=(Path="../../Plugins/Experimental/GameFeatures/Templates/GameFeaturePluginContentOnly"),Label=NSLOCTEXT("[/Script/GameFeaturesEditor]", "E17DB7E94A2C0F5968A52E89959864C0", "Game Feature (Content Only)"),Description=NSLOCTEXT("[/Script/GameFeaturesEditor]", "CB5136B64AFAFBA5EDEA078FCAC7BCC9", "Create a new Game Feature Plugin."),DefaultGameFeatureDataClass=None,DefaultGameFeatureDataName="")
-+PluginTemplates=(Path=(Path="../../Plugins/Experimental/GameFeatures/Templates/GameFeaturePluginWithCode"),Label=NSLOCTEXT("[/Script/GameFeaturesEditor]", "90D6FD2A455176CEC42DD79A23AEC3B0", "Game Feature (with C++)"),Description=NSLOCTEXT("[/Script/GameFeaturesEditor]", "2B00A0E7415C8CEAA6DA609A141237F0", "Create a new Game Feature Plugin with a minimal amount of code."),DefaultGameFeatureDataClass=None,DefaultGameFeatureDataName="")
\ No newline at end of file
++PluginTemplates=(Path=(Path="../../Plugins/Experimental/GameFeatures/Templates/GameFeaturePluginWithCode"),Label=NSLOCTEXT("[/Script/GameFeaturesEditor]", "90D6FD2A455176CEC42DD79A23AEC3B0", "Game Feature (with C++)"),Description=NSLOCTEXT("[/Script/GameFeaturesEditor]", "2B00A0E7415C8CEAA6DA609A141237F0", "Create a new Game Feature Plugin with a minimal amount of code."),DefaultGameFeatureDataClass=None,DefaultGameFeatureDataName="")
diff --git a/Engine/Plugins/Experimental/Interchange/Runtime/Source/FactoryNodes/Public/InterchangeTextureFactoryNode.h b/Engine/Plugins/Experimental/Interchange/Runtime/Source/FactoryNodes/Public/InterchangeTextureFactoryNode.h
--- a/Engine/Plugins/Experimental/Interchange/Runtime/Source/FactoryNodes/Public/InterchangeTextureFactoryNode.h
+++ b/Engine/Plugins/Experimental/Interchange/Runtime/Source/FactoryNodes/Public/InterchangeTextureFactoryNode.h
@@ -439,6 +439,25 @@ public:
 #endif
 	}
 
+	
+	/** Return false if the Attribute was not set previously.*/
+	UFUNCTION(BlueprintCallable, Category = "Interchange | Node | Texture")
+	bool GetCustombDoScaleMipsForAlphaCoverage(bool& AttributeValue) const
+	{
+		IMPLEMENT_NODE_ATTRIBUTE_GETTER(bDoScaleMipsForAlphaCoverage, bool);
+	}
+
+	UFUNCTION(BlueprintCallable, Category = "Interchange | Node | Texture")
+	bool SetCustombDoScaleMipsForAlphaCoverage(const bool& AttributeValue, bool bAddApplyDelegate = true)
+	{
+#if WITH_EDITORONLY_DATA
+		IMPLEMENT_NODE_ATTRIBUTE_SETTER(UInterchangeTextureFactoryNode, bDoScaleMipsForAlphaCoverage, bool, UTexture)
+#else
+		IMPLEMENT_NODE_ATTRIBUTE_SETTER_NODELEGATE(bDoScaleMipsForAlphaCoverage, bool)
+#endif
+	}
+
+
 	/** Return false if the Attribute was not set previously.*/
 	UFUNCTION(BlueprintCallable, Category = "Interchange | Node | Texture")
 	bool GetCustombFlipGreenChannel(bool& AttributeValue) const
@@ -778,6 +797,7 @@ private:
 
 	//Texture general
 	const UE::Interchange::FAttributeKey Macro_CustombDitherMipMapAlphaKey = UE::Interchange::FAttributeKey(TEXT("bDitherMipMapAlpha"));
+	const UE::Interchange::FAttributeKey Macro_CustombDoScaleMipsForAlphaCoverageKey = UE::Interchange::FAttributeKey(TEXT("bDoScaleMipsForAlphaCoverage"));
 	const UE::Interchange::FAttributeKey Macro_CustomAlphaCoverageThresholdsKey = UE::Interchange::FAttributeKey(TEXT("AlphaCoverageThresholds"));
 	const UE::Interchange::FAttributeKey Macro_CustombFlipGreenChannelKey = UE::Interchange::FAttributeKey(TEXT("bFlipGreenChannel"));
 	const UE::Interchange::FAttributeKey Macro_CustombForcePVRTC4Key = UE::Interchange::FAttributeKey(TEXT("bForcePVRTC4"));
diff --git a/Engine/Source/Developer/TextureBuild/Private/TextureBuildFunction.cpp b/Engine/Source/Developer/TextureBuild/Private/TextureBuildFunction.cpp
--- a/Engine/Source/Developer/TextureBuild/Private/TextureBuildFunction.cpp
+++ b/Engine/Source/Developer/TextureBuild/Private/TextureBuildFunction.cpp
@@ -27,7 +27,7 @@ DEFINE_LOG_CATEGORY_STATIC(LogTextureBuildFunction, Log, All);
 // MUST have a corresponding change to this version. Individual texture formats have a version to
 // change that is specific to the format. A merge conflict affecting the version MUST be resolved
 // by generating a new version.
-static const FGuid TextureDerivedDataVersion(TEXT("4af5c0e4-46cf-466b-a6f7-7b2b3dec8b10"));
+static const FGuid TextureDerivedDataVersion(TEXT("19ffd7ab-ea5c-4a18-b6f4-6b0b5eedd606"));
 
 #ifndef CASE_ENUM_TO_TEXT
 #define CASE_ENUM_TO_TEXT(txt) case txt: return TEXT(#txt);
@@ -103,6 +103,7 @@ static FTextureBuildSettings ReadBuildSettingsFromCompactBinary(const FCbObjectV
 	ReadCbField(ColorAdjustmentCbObj["AdjustHue"], ColorAdjustment.AdjustHue);
 	ReadCbField(ColorAdjustmentCbObj["AdjustMinAlpha"], ColorAdjustment.AdjustMinAlpha);
 	ReadCbField(ColorAdjustmentCbObj["AdjustMaxAlpha"], ColorAdjustment.AdjustMaxAlpha);
+	BuildSettings.bDoScaleMipsForAlphaCoverage = Object["bDoScaleMipsForAlphaCoverage"].AsBool(BuildSettings.bDoScaleMipsForAlphaCoverage);
 	ReadCbField(Object["AlphaCoverageThresholds"], BuildSettings.AlphaCoverageThresholds);
 	ReadCbField(Object["MipSharpening"], BuildSettings.MipSharpening);
 	ReadCbField(Object["DiffuseConvolveMipLevel"], BuildSettings.DiffuseConvolveMipLevel);
diff --git a/Engine/Source/Developer/TextureCompressor/Private/TextureCompressorModule.cpp b/Engine/Source/Developer/TextureCompressor/Private/TextureCompressorModule.cpp
--- a/Engine/Source/Developer/TextureCompressor/Private/TextureCompressorModule.cpp
+++ b/Engine/Source/Developer/TextureCompressor/Private/TextureCompressorModule.cpp
@@ -12,6 +12,7 @@
 #include "Interfaces/ITextureFormat.h"
 #include "Misc/Paths.h"
 #include "ImageCore.h"
+#include <cmath>
 
 #if PLATFORM_WINDOWS
 #include "Windows/WindowsHWrapper.h"
@@ -313,6 +314,64 @@ private:
 	float KernelWeights[MaxKernelExtend * MaxKernelExtend];
 };
 
+static float DetermineScaledThreshold(float Threshold, float Scale)
+{
+	check(Threshold > 0.f && Scale > 0.f);
+
+	// Assuming Scale > 0 and Threshold > 0, find ScaledThreshold such that
+	//	 x * Scale >= Threshold
+	// is exactly equivalent to
+	//	 x >= ScaledThreshold.
+	//
+	// This is for a test that was originally written in the first form that we want to
+	// transform to the second form without changing results (which would in turn change
+	// texture cooks).
+	//
+	// In exact arithmetic, this is just ScaledThreshold = Threshold / Scale.
+	//
+	// In floating point, we need to consider rounding. Computed in floating point
+	// and assuming round-to-nearest (breaking ties towards even), we get 
+	//
+	//	 RN(x * Scale) >= Threshold
+	//
+	// The smallest conceivable x that passes RN(x * Scale) >= Threshold is
+	// x = (Threshold - 0.5u) / Scale, landing exactly halfway with the rounding
+	// going up; this is slightly less than an exact Threshold/Scale.
+	//
+	// For regular floating point division, we get
+	//	 RN(Threshold / Scale)
+	// = (Threshold / Scale) * (1 + e),  |e| < 0.5u (the inequality is strict for divisions)
+	//
+	// That gets us relatively close to the target value, but we have no guarantee that rounding
+	// on the division was in the direction we wanted. Just check whether our target inequality
+	// is satisfied and bump up or down to the next representable float as required.
+	float ScaledThreshold = Threshold / Scale;
+	float SteppedDown = std::nextafter(ScaledThreshold, 0.f);
+
+	// We want ScaledThreshold to be the smallest float such that
+	//	 ScaledThreshold * Scale >= Threshold
+	// meaning the next-smaller float below ScaledThreshold (which is SteppedDown)
+	// should not be >=Threshold. 
+
+	if (SteppedDown * Scale >= Threshold)
+	{
+		// We were too large, go down by 1 ulp
+		ScaledThreshold = SteppedDown;
+	}
+	else if (ScaledThreshold * Scale < Threshold)
+	{
+		// We were too small, go up by 1 ulp
+		ScaledThreshold = std::nextafter(ScaledThreshold, 2.f * ScaledThreshold);
+	}
+
+	// We should now have the right threshold:
+	check(ScaledThreshold * Scale >= Threshold); // ScaledThreshold is large enough
+	check(std::nextafter(ScaledThreshold, 0.f) * Scale < Threshold); // next below is too small
+
+	return ScaledThreshold;
+}
+
+
 static FVector4f ComputeAlphaCoverage(const FVector4f& Thresholds, const FVector4f& Scales, const FImageView2D& SourceImageData)
 {
 	TRACE_CPUPROFILER_EVENT_SCOPE(ComputeAlphaCoverage);
@@ -328,7 +387,7 @@ static FVector4f ComputeAlphaCoverage(const FVector4f& Thresholds, const FVector
 		
 		check( Thresholds[3] != 0.f );
 
-		const float ThresholdScaled = Thresholds[3] / Scales[3];
+		const float ThresholdScaled = DetermineScaledThreshold(Thresholds[3] , Scales[3]);
 		
 		int32 CommonResult = 0;
 		ParallelFor(NumJobs, [&](int32 Index)
@@ -344,10 +403,7 @@ static FVector4f ComputeAlphaCoverage(const FVector4f& Thresholds, const FVector
 
 				for (int32 x = 0; x < SourceImageData.SizeX; ++x)
 				{
-					// Sample channel values at pixel neighborhood
-					const float A = RowPixels[x].A;
-
-					LocalCoverage += (A >= ThresholdScaled);
+					LocalCoverage += (RowPixels[x].A >= ThresholdScaled);
 				}
 			}
 
@@ -374,7 +430,7 @@ static FVector4f ComputeAlphaCoverage(const FVector4f& Thresholds, const FVector
 			else
 			{
 				check( Scales[i] != 0.f );
-				ThresholdsScaled[i] = Thresholds[i] / Scales[i];
+				ThresholdsScaled[i] = DetermineScaledThreshold( Thresholds[i] , Scales[i] );
 			}
 		}
 
@@ -492,6 +548,7 @@ static void GenerateSharpenedMipB8G8R8A8Templ(
 	const FImageView2D& SourceImageData, 
 	FImageView2D& DestImageData, 
 	bool bDitherMipMapAlpha,
+	bool bDoScaleMipsForAlphaCoverage,
 	FVector4f AlphaCoverages,
 	FVector4f AlphaThresholds,
 	const FImageKernel2D& Kernel,
@@ -509,7 +566,7 @@ static void GenerateSharpenedMipB8G8R8A8Templ(
 	FRandomStream RandomStream(0);
 
 	FVector4f AlphaScale(1, 1, 1, 1);
-	if (AlphaThresholds != FVector4f(0,0,0,0))
+	if (bDoScaleMipsForAlphaCoverage)
 	{
 		AlphaScale = ComputeAlphaScale(AlphaCoverages, AlphaThresholds, SourceImageData);
 	}
@@ -622,6 +679,7 @@ static void GenerateSharpenedMipB8G8R8A8(
 	FImageView2D& DestImageData, 
 	EMipGenAddressMode AddressMode, 
 	bool bDitherMipMapAlpha,
+	bool bDoScaleMipsForAlphaCoverage,
 	FVector4f AlphaCoverages,
 	FVector4f AlphaThresholds,
 	const FImageKernel2D &Kernel,
@@ -635,13 +693,13 @@ static void GenerateSharpenedMipB8G8R8A8(
 	switch(AddressMode)
 	{
 	case MGTAM_Wrap:
-		GenerateSharpenedMipB8G8R8A8Templ<MGTAM_Wrap>(SourceImageData, DestImageData, bDitherMipMapAlpha, AlphaCoverages, AlphaThresholds, Kernel, ScaleFactor, bSharpenWithoutColorShift, bUnfiltered);
+		GenerateSharpenedMipB8G8R8A8Templ<MGTAM_Wrap>(SourceImageData, DestImageData, bDitherMipMapAlpha, bDoScaleMipsForAlphaCoverage, AlphaCoverages, AlphaThresholds, Kernel, ScaleFactor, bSharpenWithoutColorShift, bUnfiltered);
 		break;
 	case MGTAM_Clamp:
-		GenerateSharpenedMipB8G8R8A8Templ<MGTAM_Clamp>(SourceImageData, DestImageData, bDitherMipMapAlpha, AlphaCoverages, AlphaThresholds, Kernel, ScaleFactor, bSharpenWithoutColorShift, bUnfiltered);
+		GenerateSharpenedMipB8G8R8A8Templ<MGTAM_Clamp>(SourceImageData, DestImageData, bDitherMipMapAlpha, bDoScaleMipsForAlphaCoverage, AlphaCoverages, AlphaThresholds, Kernel, ScaleFactor, bSharpenWithoutColorShift, bUnfiltered);
 		break;
 	case MGTAM_BorderBlack:
-		GenerateSharpenedMipB8G8R8A8Templ<MGTAM_BorderBlack>(SourceImageData, DestImageData, bDitherMipMapAlpha, AlphaCoverages, AlphaThresholds, Kernel, ScaleFactor, bSharpenWithoutColorShift, bUnfiltered);
+		GenerateSharpenedMipB8G8R8A8Templ<MGTAM_BorderBlack>(SourceImageData, DestImageData, bDitherMipMapAlpha, bDoScaleMipsForAlphaCoverage, AlphaCoverages, AlphaThresholds, Kernel, ScaleFactor, bSharpenWithoutColorShift, bUnfiltered);
 		break;
 	default:
 		check(0);
@@ -656,13 +714,13 @@ static void GenerateSharpenedMipB8G8R8A8(
 		switch(AddressMode)
 		{
 		case MGTAM_Wrap:
-			GenerateSharpenedMipB8G8R8A8Templ<MGTAM_Wrap>(SourceImageData2, TempImageData, bDitherMipMapAlpha, AlphaCoverages, AlphaThresholds, Kernel, ScaleFactor, bSharpenWithoutColorShift, bUnfiltered);
+			GenerateSharpenedMipB8G8R8A8Templ<MGTAM_Wrap>(SourceImageData2, TempImageData, bDitherMipMapAlpha, bDoScaleMipsForAlphaCoverage, AlphaCoverages, AlphaThresholds, Kernel, ScaleFactor, bSharpenWithoutColorShift, bUnfiltered);
 			break;
 		case MGTAM_Clamp:
-			GenerateSharpenedMipB8G8R8A8Templ<MGTAM_Clamp>(SourceImageData2, TempImageData, bDitherMipMapAlpha, AlphaCoverages, AlphaThresholds, Kernel, ScaleFactor, bSharpenWithoutColorShift, bUnfiltered);
+			GenerateSharpenedMipB8G8R8A8Templ<MGTAM_Clamp>(SourceImageData2, TempImageData, bDitherMipMapAlpha, bDoScaleMipsForAlphaCoverage, AlphaCoverages, AlphaThresholds, Kernel, ScaleFactor, bSharpenWithoutColorShift, bUnfiltered);
 			break;
 		case MGTAM_BorderBlack:
-			GenerateSharpenedMipB8G8R8A8Templ<MGTAM_BorderBlack>(SourceImageData2, TempImageData, bDitherMipMapAlpha, AlphaCoverages, AlphaThresholds, Kernel, ScaleFactor, bSharpenWithoutColorShift, bUnfiltered);
+			GenerateSharpenedMipB8G8R8A8Templ<MGTAM_BorderBlack>(SourceImageData2, TempImageData, bDitherMipMapAlpha, bDoScaleMipsForAlphaCoverage, AlphaCoverages, AlphaThresholds, Kernel, ScaleFactor, bSharpenWithoutColorShift, bUnfiltered);
 			break;
 		default:
 			check(0);
@@ -769,6 +827,7 @@ static void GenerateTopMip(const FImage& SrcImage, FImage& DestImage, const FTex
 			DestView,
 			AddressMode,
 			Settings.bDitherMipMapAlpha,
+			false,
 			FVector4f(0, 0, 0, 0),
 			FVector4f(0, 0, 0, 0),
 			KernelDownsample,
@@ -859,6 +918,7 @@ static void DownscaleImage(const FImage& SrcImage, FImage& DstImage, const FText
 			SrcImageData, 
 			DstImageData, 
 			Settings.bDitherMipMapAlpha, 
+			false,
 			FVector4f(0, 0, 0, 0),
 			FVector4f(0, 0, 0, 0), 
 			AvgKernel, 
@@ -982,7 +1042,7 @@ void ITextureCompressorModule::GenerateMipChain(
 	const int32 SrcHeight= BaseMip.SizeY;
 	const int32 SrcNumSlices = BaseMip.NumSlices;
 	const ERawImageFormat::Type ImageFormat = ERawImageFormat::RGBA32F;
-	FVector4f AlphaScales(1, 1, 1, 1);
+
 	FVector4f AlphaCoverages(0, 0, 0, 0);
 
 
@@ -1028,11 +1088,13 @@ void ITextureCompressorModule::GenerateMipChain(
 	}
 
 	// Calculate alpha coverage value to preserve along mip chain
-	if (Settings.AlphaCoverageThresholds != FVector4f(0,0,0,0))
+	if ( Settings.bDoScaleMipsForAlphaCoverage )
 	{
+		check(Settings.AlphaCoverageThresholds != FVector4f(0,0,0,0));
 		check(IntermediateSrcPtr);
 		const FImageView2D IntermediateSrcView = FImageView2D::ConstructConst(*IntermediateSrcPtr, 0);
-		
+
+		const FVector4f AlphaScales(1, 1, 1, 1);		
 		AlphaCoverages = ComputeAlphaCoverage(Settings.AlphaCoverageThresholds, AlphaScales, IntermediateSrcView);
 	}
 
@@ -1063,6 +1125,7 @@ void ITextureCompressorModule::GenerateMipChain(
 				DestView,
 				AddressMode,
 				Settings.bDitherMipMapAlpha,
+				Settings.bDoScaleMipsForAlphaCoverage,
 				AlphaCoverages,
 				Settings.AlphaCoverageThresholds,
 				KernelDownsample,
@@ -1080,6 +1143,7 @@ void ITextureCompressorModule::GenerateMipChain(
 					IntermediateDstView,
 					AddressMode,
 					Settings.bDitherMipMapAlpha,
+					Settings.bDoScaleMipsForAlphaCoverage,
 					AlphaCoverages,
 					Settings.AlphaCoverageThresholds,
 					KernelSimpleAverage,
diff --git a/Engine/Source/Developer/TextureCompressor/Public/TextureCompressorModule.h b/Engine/Source/Developer/TextureCompressor/Public/TextureCompressorModule.h
--- a/Engine/Source/Developer/TextureCompressor/Public/TextureCompressorModule.h
+++ b/Engine/Source/Developer/TextureCompressor/Public/TextureCompressorModule.h
@@ -72,6 +72,8 @@ struct FTextureBuildSettings
 	FCbObjectView FormatConfigOverride;
 	/** Color adjustment parameters. */
 	FColorAdjustmentParameters ColorAdjustment;
+	/** Enable preserving alpha coverage. */
+	bool bDoScaleMipsForAlphaCoverage;
 	/** Channel values to compare to when preserving alpha coverage. */
 	FVector4f AlphaCoverageThresholds;
 	/** The desired amount of mip sharpening. */
@@ -213,7 +215,8 @@ struct FTextureBuildSettings
 
 	/** Default settings. */
 	FTextureBuildSettings()
-		: AlphaCoverageThresholds(0, 0, 0, 0)
+		: bDoScaleMipsForAlphaCoverage(false)
+		, AlphaCoverageThresholds(0, 0, 0, 0)
 		, MipSharpening(0.0f)
 		, DiffuseConvolveMipLevel(0)
 		, SharpenMipKernelSize(2)
diff --git a/Engine/Source/Editor/UnrealEd/Classes/Factories/TextureFactory.h b/Engine/Source/Editor/UnrealEd/Classes/Factories/TextureFactory.h
--- a/Engine/Source/Editor/UnrealEd/Classes/Factories/TextureFactory.h
+++ b/Engine/Source/Editor/UnrealEd/Classes/Factories/TextureFactory.h
@@ -115,6 +115,10 @@ class UNREALED_API UTextureFactory : public UFactory, public IImportSettingsPars
 	/** If enabled, mip-map alpha values will be dithered for smooth transitions */
 	UPROPERTY(EditAnywhere, Category=DitherMipMaps, meta=(ToolTip="If enabled, mip-map alpha values will be dithered for smooth transitions"))
 	uint32 bDitherMipMapAlpha:1;
+	
+	/** Whether mip RGBA should be scaled to preserve the number of pixels with Value >= AlphaCoverageThresholds */
+	UPROPERTY(EditAnywhere, Category=PreserveAlphaCoverage, meta=(ToolTip="Whether mip RGBA should be scaled to preserve the number of pixels with Value >= AlphaCoverageThresholds"))
+	bool bDoScaleMipsForAlphaCoverage = false;
 
 	/** Channel values to compare to when preserving alpha coverage from a mask. */
 	UPROPERTY(EditAnywhere, Category=PreserveAlphaCoverage, meta=(ToolTip="Channel values to compare to when preserving alpha coverage from a mask for mips"))
diff --git a/Engine/Source/Editor/UnrealEd/Private/Factories/EditorFactories.cpp b/Engine/Source/Editor/UnrealEd/Private/Factories/EditorFactories.cpp
--- a/Engine/Source/Editor/UnrealEd/Private/Factories/EditorFactories.cpp
+++ b/Engine/Source/Editor/UnrealEd/Private/Factories/EditorFactories.cpp
@@ -4136,6 +4136,7 @@ UObject* UTextureFactory::FactoryCreateBinary
 	float								ExistingAdjustHue = 0.0f;
 	float								ExistingAdjustMinAlpha = 0.0f;
 	float								ExistingAdjustMaxAlpha = 1.0f;
+	bool								ExistingbDoScaleMipsForAlphaCoverage = false;
 	FVector4							ExistingAlphaCoverageThresholds = FVector4(0, 0, 0, 0);
 	TextureMipGenSettings				ExistingMipGenSettings = TextureMipGenSettings(0);
 	bool								ExistingVirtualTextureStreaming = false;
@@ -4225,6 +4226,7 @@ UObject* UTextureFactory::FactoryCreateBinary
 		ExistingDeferCompression = ExistingTexture->DeferCompression;
 		ExistingFlipGreenChannel = ExistingTexture->bFlipGreenChannel;
 		ExistingDitherMipMapAlpha = ExistingTexture->bDitherMipMapAlpha;
+		ExistingbDoScaleMipsForAlphaCoverage = ExistingTexture->bDoScaleMipsForAlphaCoverage;
 		ExistingAlphaCoverageThresholds = ExistingTexture->AlphaCoverageThresholds;
 		ExistingAdjustBrightness = ExistingTexture->AdjustBrightness;
 		ExistingAdjustBrightnessCurve = ExistingTexture->AdjustBrightnessCurve;
@@ -4345,6 +4347,7 @@ UObject* UTextureFactory::FactoryCreateBinary
 	Texture->CompressionNoAlpha		= NoAlpha;
 	Texture->DeferCompression		= bDeferCompression;
 	Texture->bDitherMipMapAlpha		= bDitherMipMapAlpha;
+	Texture->bDoScaleMipsForAlphaCoverage = bDoScaleMipsForAlphaCoverage;
 	Texture->AlphaCoverageThresholds = AlphaCoverageThresholds;
 	
 	if(Texture->MipGenSettings == TMGS_FromTextureGroup)
@@ -4407,6 +4410,7 @@ UObject* UTextureFactory::FactoryCreateBinary
 		Texture->CompressionNoAlpha = ExistingNoAlpha;
 		Texture->DeferCompression = ExistingDeferCompression;
 		Texture->bDitherMipMapAlpha = ExistingDitherMipMapAlpha;
+		Texture->bDoScaleMipsForAlphaCoverage = ExistingbDoScaleMipsForAlphaCoverage;
 		Texture->AlphaCoverageThresholds = ExistingAlphaCoverageThresholds;
 		Texture->bFlipGreenChannel = ExistingFlipGreenChannel;
 		Texture->AdjustBrightness = ExistingAdjustBrightness;
diff --git a/Engine/Source/Runtime/Core/Public/UObject/UE5MainStreamObjectVersion.h b/Engine/Source/Runtime/Core/Public/UObject/UE5MainStreamObjectVersion.h
--- a/Engine/Source/Runtime/Core/Public/UObject/UE5MainStreamObjectVersion.h
+++ b/Engine/Source/Runtime/Core/Public/UObject/UE5MainStreamObjectVersion.h
@@ -187,6 +187,9 @@ struct CORE_API FUE5MainStreamObjectVersion
 		// Added some new MeshInfo to the FSkeletalMeshLODModel class.
 		SkeletalMeshLODModelMeshInfo,
 		
+		// Add Texture DoScaleMipsForAlphaCoverage
+		TextureDoScaleMipsForAlphaCoverage,
+
 		// -----<new versions can be added above this line>-------------------------------------------------
 		VersionPlusOne,
 		LatestVersion = VersionPlusOne - 1
diff --git a/Engine/Source/Runtime/Engine/Classes/Engine/Texture.h b/Engine/Source/Runtime/Engine/Classes/Engine/Texture.h
--- a/Engine/Source/Runtime/Engine/Classes/Engine/Texture.h
+++ b/Engine/Source/Runtime/Engine/Classes/Engine/Texture.h
@@ -1064,9 +1064,13 @@ public:
 	UPROPERTY(EditAnywhere, BlueprintReadWrite, Category=Texture, AdvancedDisplay)
 	uint32 bDitherMipMapAlpha:1;
 
-	/** Alpha values per channel to compare to when preserving alpha coverage. */
-	UPROPERTY(EditAnywhere, BlueprintReadWrite, Category = Texture, meta=(ClampMin = "0", ClampMax = "1"), AdvancedDisplay)
-	FVector4 AlphaCoverageThresholds;
+	/** Whether mip RGBA should be scaled to preserve the number of pixels with Value >= AlphaCoverageThresholds.  AlphaCoverageThresholds are ignored if this is off. */
+	UPROPERTY(EditAnywhere, BlueprintReadWrite, Category=Texture, AdvancedDisplay)
+	bool bDoScaleMipsForAlphaCoverage = false;
+	
+	/** Alpha values per channel to compare to when preserving alpha coverage. 0 means disable channel.  Typical good values in 0.5 - 0.9, not 1.0 */
+	UPROPERTY(EditAnywhere, BlueprintReadWrite, Category = Texture, meta=(ClampMin = "0", ClampMax = "1.0"), AdvancedDisplay)
+	FVector4 AlphaCoverageThresholds = FVector4(0,0,0,0);
 
 	/** When true the texture's border will be preserved during mipmap generation. */
 	UPROPERTY(EditAnywhere, BlueprintReadWrite, Category=LevelOfDetail, AdvancedDisplay)
diff --git a/Engine/Source/Runtime/Engine/Private/Texture.cpp b/Engine/Source/Runtime/Engine/Private/Texture.cpp
--- a/Engine/Source/Runtime/Engine/Private/Texture.cpp
+++ b/Engine/Source/Runtime/Engine/Private/Texture.cpp
@@ -123,6 +123,7 @@ UTexture::UTexture(const FObjectInitializer& ObjectInitializer)
 	bUseLegacyGamma = false;
 	bIsImporting = false;
 	bCustomPropertiesImported = false;
+	bDoScaleMipsForAlphaCoverage = false;
 	AlphaCoverageThresholds = FVector4(0, 0, 0, 0);
 	PaddingColor = FColor::Black;
 	ChromaKeyColor = FColorList::Magenta;
@@ -576,6 +577,37 @@ void UTexture::Serialize(FArchive& Ar)
 		FWriteScopeLock BulkDataExclusiveScope(Source.BulkDataLock.Get());
 #endif
 
+		if (Ar.IsLoading() && Ar.CustomVer(FUE5MainStreamObjectVersion::GUID) < FUE5MainStreamObjectVersion::TextureDoScaleMipsForAlphaCoverage)
+		{
+			// bDoScaleMipsForAlphaCoverage was not transmitted in old versions
+			//	and AlphaCoverageThresholds was being incorrectly set to (0,0,0,1)
+			check( bDoScaleMipsForAlphaCoverage == false );
+	
+			if ( AlphaCoverageThresholds != FVector4(0,0,0,0) && AlphaCoverageThresholds != FVector4(0,0,0,1) )
+			{
+				// AlphaCoverageThresholds is a non-default value, assume that means they wanted it on
+				bDoScaleMipsForAlphaCoverage = true;
+			}
+			else if ( AlphaCoverageThresholds == FVector4(0,0,0,1) )
+			{
+				// if value is (0,0,0,1)
+				//	that was previously incorrectly being set by default and enabling alpha coverage processing
+				// we don't want that, but to optionally preserve old behavior you can set a config option :
+
+				static struct ReadConfigOnce
+				{
+					bool bBool;
+					ReadConfigOnce()
+					{
+						bBool = false;
+						GConfig->GetBool(TEXT("Texture"), TEXT("EnableLegacyAlphaCoverageThresholdScaling"), bBool, GEditorIni);
+					}
+				} ConfigValue;
+
+				bDoScaleMipsForAlphaCoverage = ConfigValue.bBool;
+			}
+		}
+
 		if (Ar.IsLoading() && Ar.CustomVer(FUE5MainStreamObjectVersion::GUID) < FUE5MainStreamObjectVersion::VirtualizedBulkDataHaveUniqueGuids)
 		{
 			if (Ar.CustomVer(FUE5MainStreamObjectVersion::GUID) < FUE5MainStreamObjectVersion::TextureSourceVirtualization)
diff --git a/Engine/Source/Runtime/Engine/Private/TextureDerivedData.cpp b/Engine/Source/Runtime/Engine/Private/TextureDerivedData.cpp
--- a/Engine/Source/Runtime/Engine/Private/TextureDerivedData.cpp
+++ b/Engine/Source/Runtime/Engine/Private/TextureDerivedData.cpp
@@ -139,8 +139,9 @@ static void SerializeForKey(FArchive& Ar, const FTextureBuildSettings& Settings)
 	TempByte = Settings.bPreserveBorder; Ar << TempByte;
 	TempByte = Settings.bDitherMipMapAlpha; Ar << TempByte;
 
-	if (Settings.AlphaCoverageThresholds != FVector4f(0, 0, 0, 0))
+	if (Settings.bDoScaleMipsForAlphaCoverage)
 	{
+		check( Settings.AlphaCoverageThresholds != FVector4f(0, 0, 0, 0) );
 		TempVector4f = Settings.AlphaCoverageThresholds; Ar << TempVector4f;
 	}
 	
@@ -692,7 +693,20 @@ static void GetTextureBuildSettings(
 	OutBuildSettings.bUseLegacyGamma = Texture.bUseLegacyGamma;
 	OutBuildSettings.bPreserveBorder = Texture.bPreserveBorder;
 	OutBuildSettings.bDitherMipMapAlpha = Texture.bDitherMipMapAlpha;
-	OutBuildSettings.AlphaCoverageThresholds = (FVector4f)Texture.AlphaCoverageThresholds;
+
+	// in Texture , the fields bDoScaleMipsForAlphaCoverage and AlphaCoverageThresholds are independent
+	// but in the BuildSettings bDoScaleMipsForAlphaCoverage is only on if thresholds are valid (not all zero)
+	if ( Texture.bDoScaleMipsForAlphaCoverage && Texture.AlphaCoverageThresholds != FVector4(0,0,0,0) )
+	{
+		OutBuildSettings.bDoScaleMipsForAlphaCoverage = Texture.bDoScaleMipsForAlphaCoverage;
+		OutBuildSettings.AlphaCoverageThresholds = (FVector4f)Texture.AlphaCoverageThresholds;
+	}
+	else
+	{
+		OutBuildSettings.bDoScaleMipsForAlphaCoverage = false;
+		OutBuildSettings.AlphaCoverageThresholds = FVector4f(0,0,0,0);
+	}
+
 	OutBuildSettings.bComputeBokehAlpha = (Texture.LODGroup == TEXTUREGROUP_Bokeh);
 	OutBuildSettings.bReplicateAlpha = false;
 	OutBuildSettings.bReplicateRed = false;
diff --git a/Engine/Source/Runtime/Engine/Private/TextureDerivedDataBuildUtils.cpp b/Engine/Source/Runtime/Engine/Private/TextureDerivedDataBuildUtils.cpp
--- a/Engine/Source/Runtime/Engine/Private/TextureDerivedDataBuildUtils.cpp
+++ b/Engine/Source/Runtime/Engine/Private/TextureDerivedDataBuildUtils.cpp
@@ -113,8 +113,13 @@ static void WriteBuildSettings(FCbWriter& Writer, const FTextureBuildSettings& B
 		WriteCbFieldWithDefault(Writer, "AdjustMaxAlpha", BuildSettings.ColorAdjustment.AdjustMaxAlpha, DefaultSettings.ColorAdjustment.AdjustMaxAlpha);
 		Writer.EndObject();
 	}
-
-	WriteCbFieldWithDefault(Writer, "AlphaCoverageThresholds", BuildSettings.AlphaCoverageThresholds, DefaultSettings.AlphaCoverageThresholds);
+	
+	WriteCbFieldWithDefault<bool>(Writer, "bDoScaleMipsForAlphaCoverage", BuildSettings.bDoScaleMipsForAlphaCoverage, DefaultSettings.bDoScaleMipsForAlphaCoverage);
+	if ( BuildSettings.bDoScaleMipsForAlphaCoverage )
+	{
+		// AlphaCoverageThresholds do not affect build if bDoScaleMipsForAlphaCoverage is off
+		WriteCbFieldWithDefault(Writer, "AlphaCoverageThresholds", BuildSettings.AlphaCoverageThresholds, DefaultSettings.AlphaCoverageThresholds);
+	}
 	WriteCbFieldWithDefault(Writer, "MipSharpening", BuildSettings.MipSharpening, DefaultSettings.MipSharpening);
 	WriteCbFieldWithDefault(Writer, "DiffuseConvolveMipLevel", BuildSettings.DiffuseConvolveMipLevel, DefaultSettings.DiffuseConvolveMipLevel);
 	WriteCbFieldWithDefault(Writer, "SharpenMipKernelSize", BuildSettings.SharpenMipKernelSize, DefaultSettings.SharpenMipKernelSize);
