TextureCompressorModule Generate Mip cleanups
special case common 2x2 box filter mip generation

diff --git a/Engine/Source/Developer/TextureCompressor/Private/TextureCompressorModule.cpp b/Engine/Source/Developer/TextureCompressor/Private/TextureCompressorModule.cpp
--- a/Engine/Source/Developer/TextureCompressor/Private/TextureCompressorModule.cpp
+++ b/Engine/Source/Developer/TextureCompressor/Private/TextureCompressorModule.cpp
@@ -135,6 +135,8 @@ public:
 		if(SharpenFactor < 0.0f)
 		{
 			// blur only
+			// this is TMGS_Blur
+			check( TableSize1D > 2 );
 			BuildGaussian1D(Table1D, TableSize1D, 1.0f, -SharpenFactor);
 			BuildFilterTable2DFrom1D(KernelWeights, Table1D, TableSize1D);
 			return;
@@ -142,12 +144,15 @@ public:
 		else if(TableSize1D == 2)
 		{
 			// 2x2 kernel: simple average
+			// SharpenFactor is ignored
+			// this is TMGS_SimpleAverage
 			KernelWeights[0] = KernelWeights[1] = KernelWeights[2] = KernelWeights[3] = 0.25f;
 			return;
 		}
 		else if(TableSize1D == 4)
 		{
 			// 4x4 kernel with sharpen or blur: can alias a bit
+			// this is not used by standard TMGS_ mip options
 			BuildFilterTable1DBase(Table1D, TableSize1D, 1.0f + SharpenFactor);
 			BuildFilterTable1DBase(NegativeTable1D, TableSize1D, -SharpenFactor);
 			BlurFilterTable1D(NegativeTable1D, TableSize1D, 1);
@@ -155,13 +160,15 @@ public:
 		else if(TableSize1D == 6)
 		{
 			// 6x6 kernel with sharpen or blur: still can alias
+			// this is not used by standard TMGS_ mip options
 			BuildFilterTable1DBase(Table1D, TableSize1D, 1.0f + SharpenFactor);
 			BuildFilterTable1DBase(NegativeTable1D, TableSize1D, -SharpenFactor);
 			BlurFilterTable1D(NegativeTable1D, TableSize1D, 2);
 		}
 		else if(TableSize1D == 8)
 		{
-			//8x8 kernel with sharpen or blur
+			//8x8 kernel with sharpen
+			// these are the TMGS_Sharpen filters
 
 			// * 2 to get similar appearance as for TableSize 6
 			SharpenFactor = SharpenFactor * 2.0f;
@@ -372,7 +379,7 @@ static float DetermineScaledThreshold(float Threshold, float Scale)
 }
 
 
-static FVector4f ComputeAlphaCoverage(const FVector4f& Thresholds, const FVector4f& Scales, const FImageView2D& SourceImageData)
+static FVector4f ComputeAlphaCoverage(const FVector4f Thresholds, const FVector4f Scales, const FImageView2D& SourceImageData)
 {
 	TRACE_CPUPROFILER_EVENT_SCOPE(ComputeAlphaCoverage);
 
@@ -477,7 +484,7 @@ static FVector4f ComputeAlphaCoverage(const FVector4f& Thresholds, const FVector
 	return Coverage;
 }
 
-static FVector4f ComputeAlphaScale(const FVector4f& Coverages, const FVector4f& AlphaThresholds, const FImageView2D& SourceImageData)
+static FVector4f ComputeAlphaScale(const FVector4f Coverages, const FVector4f AlphaThresholds, const FImageView2D& SourceImageData)
 {
 	TRACE_CPUPROFILER_EVENT_SCOPE(ComputeAlphaScale);
 
@@ -534,6 +541,33 @@ static FVector4f ComputeAlphaScale(const FVector4f& Coverages, const FVector4f&
 }
 
 
+static void GenerateMip2x2Simple(
+	const FImageView2D& SourceImageData, 
+	FImageView2D& DestImageData)
+{
+	int32 NumRowsEachJob;
+	int32 NumJobs = ImageParallelForComputeNumJobsForRows(NumRowsEachJob,DestImageData.SizeX,DestImageData.SizeY);
+
+	ParallelFor( TEXT("GenerateMip2x2Simple.PF"),NumJobs,1, [&](int32 Index)
+	{
+		int32 StartIndex = Index * NumRowsEachJob;
+		int32 EndIndex = FMath::Min(StartIndex + NumRowsEachJob, DestImageData.SizeY);
+		for (int32 DestY = StartIndex; DestY < EndIndex; ++DestY)
+		{
+			FLinearColor * DestRow = &DestImageData.Access(0,DestY);
+			const FLinearColor * SourceRow0 = &SourceImageData.Access(0,2*DestY);
+			const FLinearColor * SourceRow1 = &SourceImageData.Access(0,2*DestY+1);
+
+			for ( int32 DestX = 0;DestX < DestImageData.SizeX; DestX++ )
+			{
+				DestRow[DestX] = (SourceRow0[0] + SourceRow0[1] + SourceRow1[0] + SourceRow1[1]) * 0.25f;
+				SourceRow0 += 2;
+				SourceRow1 += 2;
+			}
+		}
+	});
+}
+
 /**
 * Generates a mip-map for an 2D B8G8R8A8 image using a 4x4 filter with sharpening
 * @param SourceImageData - The source image's data.
@@ -549,8 +583,8 @@ static void GenerateSharpenedMipB8G8R8A8Templ(
 	FImageView2D& DestImageData, 
 	bool bDitherMipMapAlpha,
 	bool bDoScaleMipsForAlphaCoverage,
-	FVector4f AlphaCoverages,
-	FVector4f AlphaThresholds,
+	const FVector4f AlphaCoverages,
+	const FVector4f AlphaThresholds,
 	const FImageKernel2D& Kernel,
 	uint32 ScaleFactor,
 	bool bSharpenWithoutColorShift,
@@ -558,11 +592,33 @@ static void GenerateSharpenedMipB8G8R8A8Templ(
 {
 	check( SourceImageData.SizeX == ScaleFactor * DestImageData.SizeX || DestImageData.SizeX == 1 );
 	check( SourceImageData.SizeY == ScaleFactor * DestImageData.SizeY || DestImageData.SizeY == 1 );
-	checkf( Kernel.GetFilterTableSize() >= 2, TEXT("Kernel table size %d, expected at least 2!"), Kernel.GetFilterTableSize());
+	
+	int32 KernelFilterTableSize = (int32) Kernel.GetFilterTableSize();
+
+	checkf( KernelFilterTableSize >= 2, TEXT("Kernel table size %d, expected at least 2!"), KernelFilterTableSize);
+	if ( KernelFilterTableSize == 2 )
+	{
+		// 2x2 is always box filter
+		check( Kernel.GetAt(0,0) == 0.25f );
+	}
+	
+	if ( KernelFilterTableSize == 2 &&
+		ScaleFactor == 2 &&
+		DestImageData.SizeX*2 == SourceImageData.SizeX &&
+		DestImageData.SizeY*2 == SourceImageData.SizeY &&
+		! bDoScaleMipsForAlphaCoverage &&
+		! bUnfiltered &&
+		! bDitherMipMapAlpha )
+	{
+		// bSharpenWithoutColorShift is ignored for 2x2 filter
+		GenerateMip2x2Simple(SourceImageData,DestImageData);
+		return;
+	}
 
-	const int32 KernelCenter = (int32)Kernel.GetFilterTableSize() / 2 - 1;
+	const int32 KernelCenter = KernelFilterTableSize / 2 - 1;
 
 	// Set up a random number stream for dithering.
+	// @todo Oodle : this is shared across the ParallelFor which is very bad
 	FRandomStream RandomStream(0);
 
 	FVector4f AlphaScale(1, 1, 1, 1);
@@ -574,99 +630,110 @@ static void GenerateSharpenedMipB8G8R8A8Templ(
 	int32 NumRowsEachJob;
 	int32 NumJobs = ImageParallelForComputeNumJobsForRows(NumRowsEachJob,DestImageData.SizeX,DestImageData.SizeY);
 
-	ParallelFor(NumJobs, [&](int32 Index)
+	ParallelFor( TEXT("GenerateSharpenedMip.PF"),NumJobs,1, [&](int32 Index)
 	{
-		TRACE_CPUPROFILER_EVENT_SCOPE(GenerateSharpenedMip.PF);
-
 		int32 StartIndex = Index * NumRowsEachJob;
 		int32 EndIndex = FMath::Min(StartIndex + NumRowsEachJob, DestImageData.SizeY);
 		for (int32 DestY = StartIndex; DestY < EndIndex; ++DestY)
 		{
-		for ( int32 DestX = 0;DestX < DestImageData.SizeX; DestX++ )
-		{
-			const int32 SourceX = DestX * ScaleFactor;
-			const int32 SourceY = DestY * ScaleFactor;
-
-			FLinearColor FilteredColor(0, 0, 0, 0);
-
-			if ( bUnfiltered )
+			for ( int32 DestX = 0;DestX < DestImageData.SizeX; DestX++ )
 			{
-				FilteredColor = LookupSourceMip<AddressMode>(SourceImageData, SourceX + 0, SourceY + 0);
-			}
-			else if ( bSharpenWithoutColorShift )
-			{
-				FLinearColor SharpenedColor(0, 0, 0, 0);
+				const int32 SourceX = DestX * ScaleFactor;
+				const int32 SourceY = DestY * ScaleFactor;
+
+				FLinearColor FilteredColor(0, 0, 0, 0);
 
-				for ( uint32 KernelY = 0; KernelY < Kernel.GetFilterTableSize();  ++KernelY )
+				if ( bUnfiltered )
+				{
+					FilteredColor = LookupSourceMip<AddressMode>(SourceImageData, SourceX + 0, SourceY + 0);
+				}
+				else if ( KernelFilterTableSize == 2 )
+				{
+					// simple 2x2 kernel to compute the color
+					// check for this case early so that bSharpenWithoutColorShift isn't run on 2x2
+
+					FilteredColor =
+						( LookupSourceMip<AddressMode>( SourceImageData, SourceX + 0, SourceY + 0 )
+						+ LookupSourceMip<AddressMode>( SourceImageData, SourceX + 1, SourceY + 0 )
+						+ LookupSourceMip<AddressMode>( SourceImageData, SourceX + 0, SourceY + 1 )
+						+ LookupSourceMip<AddressMode>( SourceImageData, SourceX + 1, SourceY + 1 ) ) * 0.25f;
+				}
+				else if ( bSharpenWithoutColorShift )
 				{
-					for ( uint32 KernelX = 0; KernelX < Kernel.GetFilterTableSize();  ++KernelX )
+					// don't use this if KernelFilterTableSize == 2 (box)
+					//	because it just does the same thing twice
+
+					FLinearColor SharpenedColor(0, 0, 0, 0);
+
+					for ( int32 KernelY = 0; KernelY < KernelFilterTableSize;  ++KernelY )
 					{
-						float Weight = Kernel.GetAt( KernelX, KernelY );
-						FLinearColor Sample = LookupSourceMip<AddressMode>( SourceImageData, SourceX + KernelX - KernelCenter, SourceY + KernelY - KernelCenter );
-						SharpenedColor += Weight * Sample;
+						for ( int32 KernelX = 0; KernelX < KernelFilterTableSize;  ++KernelX )
+						{
+							float Weight = Kernel.GetAt( KernelX, KernelY );
+							FLinearColor Sample = LookupSourceMip<AddressMode>( SourceImageData, SourceX + KernelX - KernelCenter, SourceY + KernelY - KernelCenter );
+							SharpenedColor += Weight * Sample;
+						}
 					}
-				}
 
-				float NewLuminance = SharpenedColor.GetLuminance();
+					float NewLuminance = SharpenedColor.GetLuminance();
 
-				// simple 2x2 kernel to compute the color
-				FilteredColor =
-					( LookupSourceMip<AddressMode>( SourceImageData, SourceX + 0, SourceY + 0 )
-					+ LookupSourceMip<AddressMode>( SourceImageData, SourceX + 1, SourceY + 0 )
-					+ LookupSourceMip<AddressMode>( SourceImageData, SourceX + 0, SourceY + 1 )
-					+ LookupSourceMip<AddressMode>( SourceImageData, SourceX + 1, SourceY + 1 ) ) * 0.25f;
+					// simple 2x2 kernel to compute the color
+					FilteredColor =
+						( LookupSourceMip<AddressMode>( SourceImageData, SourceX + 0, SourceY + 0 )
+						+ LookupSourceMip<AddressMode>( SourceImageData, SourceX + 1, SourceY + 0 )
+						+ LookupSourceMip<AddressMode>( SourceImageData, SourceX + 0, SourceY + 1 )
+						+ LookupSourceMip<AddressMode>( SourceImageData, SourceX + 1, SourceY + 1 ) ) * 0.25f;
 
-				float OldLuminance = FilteredColor.GetLuminance();
+					float OldLuminance = FilteredColor.GetLuminance();
 
-				if ( OldLuminance > 0.001f )
-				{
-					float Factor = NewLuminance / OldLuminance;
-					FilteredColor.R *= Factor;
-					FilteredColor.G *= Factor;
-					FilteredColor.B *= Factor;
-				}
+					if ( OldLuminance > 0.001f )
+					{
+						float Factor = NewLuminance / OldLuminance;
+						FilteredColor.R *= Factor;
+						FilteredColor.G *= Factor;
+						FilteredColor.B *= Factor;
+					}
 
-				// We also want to sharpen the alpha channel (was missing before)
-				FilteredColor.A = SharpenedColor.A;
-			}
-			else
-			{
-				for ( uint32 KernelY = 0; KernelY < Kernel.GetFilterTableSize();  ++KernelY )
+					// We also want to sharpen the alpha channel (was missing before)
+					FilteredColor.A = SharpenedColor.A;
+				}
+				else
 				{
-					for ( uint32 KernelX = 0; KernelX < Kernel.GetFilterTableSize();  ++KernelX )
+					for ( int32 KernelY = 0; KernelY < KernelFilterTableSize;  ++KernelY )
 					{
-						float Weight = Kernel.GetAt( KernelX, KernelY );
-						FLinearColor Sample = LookupSourceMip<AddressMode>( SourceImageData, SourceX + KernelX - KernelCenter, SourceY + KernelY - KernelCenter );
-						FilteredColor += Weight	* Sample;
+						for ( int32 KernelX = 0; KernelX < KernelFilterTableSize;  ++KernelX )
+						{
+							float Weight = Kernel.GetAt( KernelX, KernelY );
+							FLinearColor Sample = LookupSourceMip<AddressMode>( SourceImageData, SourceX + KernelX - KernelCenter, SourceY + KernelY - KernelCenter );
+							FilteredColor += Weight	* Sample;
+						}
 					}
 				}
-			}
-
-			// Apply computed alpha scales to each channel		
-			FilteredColor.R *= AlphaScale.X;
-			FilteredColor.G *= AlphaScale.Y;
-			FilteredColor.B *= AlphaScale.Z;
-			FilteredColor.A *= AlphaScale.W;
 
+				// Apply computed alpha scales to each channel		
+				FilteredColor.R *= AlphaScale.X;
+				FilteredColor.G *= AlphaScale.Y;
+				FilteredColor.B *= AlphaScale.Z;
+				FilteredColor.A *= AlphaScale.W;
 
-			if ( bDitherMipMapAlpha )
-			{
-				// Dither the alpha of any pixel which passes an alpha threshold test.
-				const int32 DitherAlphaThreshold = 5.0f / 255.0f;
-				const float MinRandomAlpha = 85.0f;
-				const float MaxRandomAlpha = 255.0f;
-
-				if ( FilteredColor.A > DitherAlphaThreshold)
+				if ( bDitherMipMapAlpha )
 				{
-					FilteredColor.A = FMath::TruncToInt( FMath::Lerp( MinRandomAlpha, MaxRandomAlpha, RandomStream.GetFraction() ) );
+					// @todo Oodle : this looks very odd
+					// 
+					// Dither the alpha of any pixel which passes an alpha threshold test.
+					const int32 DitherAlphaThreshold = 5.0f / 255.0f;
+					const float MinRandomAlpha = 85.0f;
+					const float MaxRandomAlpha = 255.0f;
+
+					if ( FilteredColor.A > DitherAlphaThreshold)
+					{
+						FilteredColor.A = FMath::TruncToInt( FMath::Lerp( MinRandomAlpha, MaxRandomAlpha, RandomStream.GetFraction() ) );
+					}
 				}
-			}
 
-			// Set the destination pixel.
-			//FLinearColor& DestColor = *(DestImageData.AsRGBA32F() + DestX + DestY * DestImageData.SizeX);
-			FLinearColor& DestColor = DestImageData.Access(DestX, DestY);
-			DestColor = FilteredColor;
-		}
+				// Set the destination pixel.
+				DestImageData.Access(DestX, DestY) = FilteredColor;
+			}
 		}
 	});
 }
@@ -729,8 +796,8 @@ static void GenerateSharpenedMipB8G8R8A8(
 		const int32 NumColors = DestImageData.SizeX * DestImageData.SizeY;
 		for (int32 ColorIndex = 0; ColorIndex < NumColors; ++ColorIndex)
 		{
-			DestImageData.SliceColors[ColorIndex] += TempImageData.SliceColors[ColorIndex];
-			DestImageData.SliceColors[ColorIndex] *= .5;
+			DestImageData.SliceColors[ColorIndex] =
+				(DestImageData.SliceColors[ColorIndex] + TempImageData.SliceColors[ColorIndex]) * 0.5f;
 		}
 	}
 }
@@ -1043,9 +1110,6 @@ void ITextureCompressorModule::GenerateMipChain(
 	const int32 SrcNumSlices = BaseMip.NumSlices;
 	const ERawImageFormat::Type ImageFormat = ERawImageFormat::RGBA32F;
 
-	FVector4f AlphaCoverages(0, 0, 0, 0);
-
-
 	const FImage* IntermediateSrcPtr;
 	FImage* IntermediateDstPtr;
 
@@ -1088,6 +1152,7 @@ void ITextureCompressorModule::GenerateMipChain(
 	}
 
 	// Calculate alpha coverage value to preserve along mip chain
+	FVector4f AlphaCoverages(0, 0, 0, 0);
 	if ( Settings.bDoScaleMipsForAlphaCoverage )
 	{
 		check(Settings.AlphaCoverageThresholds != FVector4f(0,0,0,0));
diff --git a/Engine/Source/Runtime/Engine/Private/TextureLODSettings.cpp b/Engine/Source/Runtime/Engine/Private/TextureLODSettings.cpp
--- a/Engine/Source/Runtime/Engine/Private/TextureLODSettings.cpp
+++ b/Engine/Source/Runtime/Engine/Private/TextureLODSettings.cpp
@@ -292,6 +292,7 @@ void UTextureLODSettings::GetMipGenSettings(const UTexture& Texture, TextureMipG
 	else if(Setting >= TMGS_Blur1 && Setting <= TMGS_Blur5)
 	{
 		int32 BlurFactor = ((int32)Setting + 1 - (int32)TMGS_Blur1);
+		check( BlurFactor > 0 );
 		OutSharpen = -BlurFactor * 2;
 		OutKernelSize = 2 + 2 * BlurFactor;
 		bOutDownsampleWithAverage = false;
