Remove broken bDitherMipMapAlpha
Wasn't doing anything useful
It actually converted any non-zero alpha to full opaque
almost the opposite of what it claimed to do
Simplify the option set by removing useless options rather than fix it

diff --git a/Engine/Plugins/Experimental/Interchange/Runtime/Source/FactoryNodes/Public/InterchangeTextureFactoryNode.h b/Engine/Plugins/Experimental/Interchange/Runtime/Source/FactoryNodes/Public/InterchangeTextureFactoryNode.h
--- a/Engine/Plugins/Experimental/Interchange/Runtime/Source/FactoryNodes/Public/InterchangeTextureFactoryNode.h
+++ b/Engine/Plugins/Experimental/Interchange/Runtime/Source/FactoryNodes/Public/InterchangeTextureFactoryNode.h
@@ -397,23 +397,6 @@ public:
 	//////////////////////////////////////////////////////////////////////////
 	//Texture general
 
-	/** Return false if the Attribute was not set previously.*/
-	UFUNCTION(BlueprintCallable, Category = "Interchange | Node | Texture")
-	bool GetCustombDitherMipMapAlpha(bool& AttributeValue) const
-	{
-		IMPLEMENT_NODE_ATTRIBUTE_GETTER(bDitherMipMapAlpha, bool);
-	}
-
-	UFUNCTION(BlueprintCallable, Category = "Interchange | Node | Texture")
-	bool SetCustombDitherMipMapAlpha(const bool& AttributeValue, bool bAddApplyDelegate = true)
-	{
-#if WITH_EDITORONLY_DATA
-		IMPLEMENT_NODE_ATTRIBUTE_SETTER(UInterchangeTextureFactoryNode, bDitherMipMapAlpha, bool, UTexture)
-#else
-		IMPLEMENT_NODE_ATTRIBUTE_SETTER_NODELEGATE(bDitherMipMapAlpha, bool)
-#endif
-	}
-
 	/** Return false if the Attribute was not set previously.*/
 	UFUNCTION(BlueprintCallable, Category = "Interchange | Node | Texture")
 	bool GetCustomAlphaCoverageThresholds(FVector4& AttributeValue) const
@@ -788,7 +771,6 @@ private:
 	const UE::Interchange::FAttributeKey Macro_CustomCompressionSettingsKey = UE::Interchange::FAttributeKey(TEXT("CompressionSettings"));
 
 	//Texture general
-	const UE::Interchange::FAttributeKey Macro_CustombDitherMipMapAlphaKey = UE::Interchange::FAttributeKey(TEXT("bDitherMipMapAlpha"));
 	const UE::Interchange::FAttributeKey Macro_CustombDoScaleMipsForAlphaCoverageKey = UE::Interchange::FAttributeKey(TEXT("bDoScaleMipsForAlphaCoverage"));
 	const UE::Interchange::FAttributeKey Macro_CustomAlphaCoverageThresholdsKey = UE::Interchange::FAttributeKey(TEXT("AlphaCoverageThresholds"));
 	const UE::Interchange::FAttributeKey Macro_CustombFlipGreenChannelKey = UE::Interchange::FAttributeKey(TEXT("bFlipGreenChannel"));
diff --git a/Engine/Source/Developer/TextureBuild/Private/TextureBuildFunction.cpp b/Engine/Source/Developer/TextureBuild/Private/TextureBuildFunction.cpp
--- a/Engine/Source/Developer/TextureBuild/Private/TextureBuildFunction.cpp
+++ b/Engine/Source/Developer/TextureBuild/Private/TextureBuildFunction.cpp
@@ -128,7 +128,6 @@ static FTextureBuildSettings ReadBuildSettingsFromCompactBinary(const FCbObjectV
 	BuildSettings.bPreserveBorder = Object["bPreserveBorder"].AsBool(BuildSettings.bPreserveBorder);
 	BuildSettings.bForceNoAlphaChannel = Object["bForceNoAlphaChannel"].AsBool(BuildSettings.bForceNoAlphaChannel);
 	BuildSettings.bForceAlphaChannel = Object["bForceAlphaChannel"].AsBool(BuildSettings.bForceAlphaChannel);
-	BuildSettings.bDitherMipMapAlpha = Object["bDitherMipMapAlpha"].AsBool(BuildSettings.bDitherMipMapAlpha);
 	BuildSettings.bComputeBokehAlpha = Object["bComputeBokehAlpha"].AsBool(BuildSettings.bComputeBokehAlpha);
 	BuildSettings.bReplicateRed = Object["bReplicateRed"].AsBool(BuildSettings.bReplicateRed);
 	BuildSettings.bReplicateAlpha = Object["bReplicateAlpha"].AsBool(BuildSettings.bReplicateAlpha);
diff --git a/Engine/Source/Developer/TextureCompressor/Private/TextureCompressorModule.cpp b/Engine/Source/Developer/TextureCompressor/Private/TextureCompressorModule.cpp
--- a/Engine/Source/Developer/TextureCompressor/Private/TextureCompressorModule.cpp
+++ b/Engine/Source/Developer/TextureCompressor/Private/TextureCompressorModule.cpp
@@ -397,10 +397,8 @@ static FVector4f ComputeAlphaCoverage(const FVector4f Thresholds, const FVector4
 		const float ThresholdScaled = DetermineScaledThreshold(Thresholds[3] , Scales[3]);
 		
 		int32 CommonResult = 0;
-		ParallelFor(NumJobs, [&](int32 Index)
+		ParallelFor( TEXT("ComputeAlphaCoverage.PF"),NumJobs,1, [&](int32 Index)
 		{
-			TRACE_CPUPROFILER_EVENT_SCOPE(ComputeAlphaCoverage.PF);
-
 			int32 StartIndex = Index * NumRowsEachJob;
 			int32 EndIndex = FMath::Min(StartIndex + NumRowsEachJob, SourceImageData.SizeY);
 			int32 LocalCoverage = 0;
@@ -442,10 +440,8 @@ static FVector4f ComputeAlphaCoverage(const FVector4f Thresholds, const FVector4
 		}
 
 		int32 CommonResults[4] = { 0, 0, 0, 0 };
-		ParallelFor(NumJobs, [&](int32 Index)
+		ParallelFor( TEXT("ComputeAlphaCoverage.PF"),NumJobs,1, [&](int32 Index)
 		{
-			TRACE_CPUPROFILER_EVENT_SCOPE(ComputeAlphaCoverage.PF);
-
 			int32 StartIndex = Index * NumRowsEachJob;
 			int32 EndIndex = FMath::Min(StartIndex + NumRowsEachJob, SourceImageData.SizeY);
 			int32 LocalCoverage[4] = { 0, 0, 0, 0 };
@@ -581,7 +577,6 @@ template <EMipGenAddressMode AddressMode>
 static void GenerateSharpenedMipB8G8R8A8Templ(
 	const FImageView2D& SourceImageData, 
 	FImageView2D& DestImageData, 
-	bool bDitherMipMapAlpha,
 	bool bDoScaleMipsForAlphaCoverage,
 	const FVector4f AlphaCoverages,
 	const FVector4f AlphaThresholds,
@@ -607,8 +602,7 @@ static void GenerateSharpenedMipB8G8R8A8Templ(
 		DestImageData.SizeX*2 == SourceImageData.SizeX &&
 		DestImageData.SizeY*2 == SourceImageData.SizeY &&
 		! bDoScaleMipsForAlphaCoverage &&
-		! bUnfiltered &&
-		! bDitherMipMapAlpha )
+		! bUnfiltered )
 	{
 		// bSharpenWithoutColorShift is ignored for 2x2 filter
 		GenerateMip2x2Simple(SourceImageData,DestImageData);
@@ -617,10 +611,6 @@ static void GenerateSharpenedMipB8G8R8A8Templ(
 
 	const int32 KernelCenter = KernelFilterTableSize / 2 - 1;
 
-	// Set up a random number stream for dithering.
-	// @todo Oodle : this is shared across the ParallelFor which is very bad
-	FRandomStream RandomStream(0);
-
 	FVector4f AlphaScale(1, 1, 1, 1);
 	if (bDoScaleMipsForAlphaCoverage)
 	{
@@ -716,21 +706,6 @@ static void GenerateSharpenedMipB8G8R8A8Templ(
 				FilteredColor.B *= AlphaScale.Z;
 				FilteredColor.A *= AlphaScale.W;
 
-				if ( bDitherMipMapAlpha )
-				{
-					// @todo Oodle : this looks very odd
-					// 
-					// Dither the alpha of any pixel which passes an alpha threshold test.
-					const int32 DitherAlphaThreshold = 5.0f / 255.0f;
-					const float MinRandomAlpha = 85.0f;
-					const float MaxRandomAlpha = 255.0f;
-
-					if ( FilteredColor.A > DitherAlphaThreshold)
-					{
-						FilteredColor.A = FMath::TruncToInt( FMath::Lerp( MinRandomAlpha, MaxRandomAlpha, RandomStream.GetFraction() ) );
-					}
-				}
-
 				// Set the destination pixel.
 				DestImageData.Access(DestX, DestY) = FilteredColor;
 			}
@@ -745,7 +720,6 @@ static void GenerateSharpenedMipB8G8R8A8(
 	const FImageView2D& SourceImageData2, // Only used with volume texture.
 	FImageView2D& DestImageData, 
 	EMipGenAddressMode AddressMode, 
-	bool bDitherMipMapAlpha,
 	bool bDoScaleMipsForAlphaCoverage,
 	FVector4f AlphaCoverages,
 	FVector4f AlphaThresholds,
@@ -760,13 +734,13 @@ static void GenerateSharpenedMipB8G8R8A8(
 	switch(AddressMode)
 	{
 	case MGTAM_Wrap:
-		GenerateSharpenedMipB8G8R8A8Templ<MGTAM_Wrap>(SourceImageData, DestImageData, bDitherMipMapAlpha, bDoScaleMipsForAlphaCoverage, AlphaCoverages, AlphaThresholds, Kernel, ScaleFactor, bSharpenWithoutColorShift, bUnfiltered);
+		GenerateSharpenedMipB8G8R8A8Templ<MGTAM_Wrap>(SourceImageData, DestImageData, bDoScaleMipsForAlphaCoverage, AlphaCoverages, AlphaThresholds, Kernel, ScaleFactor, bSharpenWithoutColorShift, bUnfiltered);
 		break;
 	case MGTAM_Clamp:
-		GenerateSharpenedMipB8G8R8A8Templ<MGTAM_Clamp>(SourceImageData, DestImageData, bDitherMipMapAlpha, bDoScaleMipsForAlphaCoverage, AlphaCoverages, AlphaThresholds, Kernel, ScaleFactor, bSharpenWithoutColorShift, bUnfiltered);
+		GenerateSharpenedMipB8G8R8A8Templ<MGTAM_Clamp>(SourceImageData, DestImageData, bDoScaleMipsForAlphaCoverage, AlphaCoverages, AlphaThresholds, Kernel, ScaleFactor, bSharpenWithoutColorShift, bUnfiltered);
 		break;
 	case MGTAM_BorderBlack:
-		GenerateSharpenedMipB8G8R8A8Templ<MGTAM_BorderBlack>(SourceImageData, DestImageData, bDitherMipMapAlpha, bDoScaleMipsForAlphaCoverage, AlphaCoverages, AlphaThresholds, Kernel, ScaleFactor, bSharpenWithoutColorShift, bUnfiltered);
+		GenerateSharpenedMipB8G8R8A8Templ<MGTAM_BorderBlack>(SourceImageData, DestImageData, bDoScaleMipsForAlphaCoverage, AlphaCoverages, AlphaThresholds, Kernel, ScaleFactor, bSharpenWithoutColorShift, bUnfiltered);
 		break;
 	default:
 		check(0);
@@ -781,13 +755,13 @@ static void GenerateSharpenedMipB8G8R8A8(
 		switch(AddressMode)
 		{
 		case MGTAM_Wrap:
-			GenerateSharpenedMipB8G8R8A8Templ<MGTAM_Wrap>(SourceImageData2, TempImageData, bDitherMipMapAlpha, bDoScaleMipsForAlphaCoverage, AlphaCoverages, AlphaThresholds, Kernel, ScaleFactor, bSharpenWithoutColorShift, bUnfiltered);
+			GenerateSharpenedMipB8G8R8A8Templ<MGTAM_Wrap>(SourceImageData2, TempImageData, bDoScaleMipsForAlphaCoverage, AlphaCoverages, AlphaThresholds, Kernel, ScaleFactor, bSharpenWithoutColorShift, bUnfiltered);
 			break;
 		case MGTAM_Clamp:
-			GenerateSharpenedMipB8G8R8A8Templ<MGTAM_Clamp>(SourceImageData2, TempImageData, bDitherMipMapAlpha, bDoScaleMipsForAlphaCoverage, AlphaCoverages, AlphaThresholds, Kernel, ScaleFactor, bSharpenWithoutColorShift, bUnfiltered);
+			GenerateSharpenedMipB8G8R8A8Templ<MGTAM_Clamp>(SourceImageData2, TempImageData, bDoScaleMipsForAlphaCoverage, AlphaCoverages, AlphaThresholds, Kernel, ScaleFactor, bSharpenWithoutColorShift, bUnfiltered);
 			break;
 		case MGTAM_BorderBlack:
-			GenerateSharpenedMipB8G8R8A8Templ<MGTAM_BorderBlack>(SourceImageData2, TempImageData, bDitherMipMapAlpha, bDoScaleMipsForAlphaCoverage, AlphaCoverages, AlphaThresholds, Kernel, ScaleFactor, bSharpenWithoutColorShift, bUnfiltered);
+			GenerateSharpenedMipB8G8R8A8Templ<MGTAM_BorderBlack>(SourceImageData2, TempImageData, bDoScaleMipsForAlphaCoverage, AlphaCoverages, AlphaThresholds, Kernel, ScaleFactor, bSharpenWithoutColorShift, bUnfiltered);
 			break;
 		default:
 			check(0);
@@ -893,7 +867,6 @@ static void GenerateTopMip(const FImage& SrcImage, FImage& DestImage, const FTex
 			FImageView2D(),
 			DestView,
 			AddressMode,
-			Settings.bDitherMipMapAlpha,
 			false,
 			FVector4f(0, 0, 0, 0),
 			FVector4f(0, 0, 0, 0),
@@ -930,7 +903,6 @@ struct FTextureDownscaleSettings
 	int32 BlockSize;
 	float Downscale;
 	uint8 DownscaleOptions;
-	bool bDitherMipMapAlpha;
 };
 
 static void DownscaleImage(const FImage& SrcImage, FImage& DstImage, const FTextureDownscaleSettings& Settings)
@@ -984,7 +956,6 @@ static void DownscaleImage(const FImage& SrcImage, FImage& DstImage, const FText
 		GenerateSharpenedMipB8G8R8A8Templ<MGTAM_Clamp>(
 			SrcImageData, 
 			DstImageData, 
-			Settings.bDitherMipMapAlpha, 
 			false,
 			FVector4f(0, 0, 0, 0),
 			FVector4f(0, 0, 0, 0), 
@@ -1032,8 +1003,6 @@ static void DownscaleImage(const FImage& SrcImage, FImage& DstImage, const FText
 		ImageChain[0] = &Image0;
 	}
 	
-	// Set up a random number stream for dithering.
-	FRandomStream RandomStream(0);
 	ImageChain[1]->Init(FinalSizeX, FinalSizeY, ImageChain[0]->NumSlices, ImageChain[0]->Format, ImageChain[0]->GammaSpace);
 	Downscale = (float)ImageChain[0]->SizeX / FinalSizeX;
 
@@ -1073,19 +1042,6 @@ static void DownscaleImage(const FImage& SrcImage, FImage& DstImage, const FText
 				}
 			}
 
-			if (Settings.bDitherMipMapAlpha)
-			{
-				// Dither the alpha of any pixel which passes an alpha threshold test.
-				const int32 DitherAlphaThreshold = 5.0f / 255.0f;
-				const float MinRandomAlpha = 85.0f;
-				const float MaxRandomAlpha = 255.0f;
-
-				if (FilteredColor.A > DitherAlphaThreshold)
-				{
-					FilteredColor.A = FMath::TruncToInt(FMath::Lerp(MinRandomAlpha, MaxRandomAlpha, RandomStream.GetFraction()));
-				}
-			}
-
 			// Set the destination pixel.
 			FLinearColor& DestColor = DstImageData.Access(X, Y);
 			DestColor = FilteredColor;
@@ -1189,7 +1145,6 @@ void ITextureCompressorModule::GenerateMipChain(
 				IntermediateSrcView2,
 				DestView,
 				AddressMode,
-				Settings.bDitherMipMapAlpha,
 				Settings.bDoScaleMipsForAlphaCoverage,
 				AlphaCoverages,
 				Settings.AlphaCoverageThresholds,
@@ -1207,7 +1162,6 @@ void ITextureCompressorModule::GenerateMipChain(
 					IntermediateSrcView2,
 					IntermediateDstView,
 					AddressMode,
-					Settings.bDitherMipMapAlpha,
 					Settings.bDoScaleMipsForAlphaCoverage,
 					AlphaCoverages,
 					Settings.AlphaCoverageThresholds,
@@ -1822,17 +1776,9 @@ void ITextureCompressorModule::AdjustImageColors(FImage& Image, const FTextureBu
 		int64 NumPixelsEachJob;
 		int32 NumJobs = ImageParallelForComputeNumJobsForPixels(NumPixelsEachJob,NumPixels);
 
-		// bForceSingleThread is set to true when: 
-		// editor or cooker is loading as this is when the derived data cache is rebuilt as it will already be limited to a single thread 
-		//     and thus overhead of multithreading will simply make it slower
-		bool bForceSingleThread = GIsEditorLoadingPackage || GIsCookerLoadingPackage || IsInAsyncLoadingThread();
-
-		// TFunction or auto are okay here
-		// TFunctionRef is not
-		TFunction<void (int32)> AdjustImageColorsFunc = [&](int32 Index)
+		//TFunction<void (int32)> 
+		auto AdjustImageColorsFunc = [&](int32 Index)
 		{
-			TRACE_CPUPROFILER_EVENT_SCOPE(AdjustImageColors.PF);
-
 			int64 StartIndex = Index * NumPixelsEachJob;
 			int64 EndIndex = FMath::Min(StartIndex + NumPixelsEachJob, NumPixels);
 			for (int64 CurPixelIndex = StartIndex; CurPixelIndex < EndIndex; ++CurPixelIndex)
@@ -1924,8 +1870,15 @@ void ITextureCompressorModule::AdjustImageColors(FImage& Image, const FTextureBu
 				ImageColors[CurPixelIndex] = LinearColor;
 			}
 		};
+		
+		// bForceSingleThread is set to true when: 
+		// editor or cooker is loading as this is when the derived data cache is rebuilt as it will already be limited to a single thread 
+		//     and thus overhead of multithreading will simply make it slower
+		// @todo Oodle - this is done here and not in other similar ParallelFor places here.  It should either be done everywhere or nowhere.
+		bool bForceSingleThread = GIsEditorLoadingPackage || GIsCookerLoadingPackage || IsInAsyncLoadingThread();
 
-		ParallelFor(NumJobs, AdjustImageColorsFunc, bForceSingleThread);
+		ParallelFor( TEXT("AdjustImageColorsFunc.PF"),NumJobs,1, AdjustImageColorsFunc, 
+			(bForceSingleThread ? EParallelForFlags::ForceSingleThread : EParallelForFlags::None) );
 	}
 }
 
@@ -2335,25 +2288,27 @@ static bool CompressMipChain(
 		}
 	}
 
-	ParallelForWithPreWork(AsyncCompressionTasks.Num(), [&AsyncCompressionTasks](int32 TaskIndex)
-	{
-		AsyncCompressionTasks[TaskIndex].DoWork();
-	},
-	[&PreWorkTasks, &TextureFormat, &OutMips, &bCompressionSucceeded, &CompressorCaps, &Settings, DebugTexturePathName, FirstMipTailIndex, bImageHasAlphaChannel]()
-	{
-		for (PreWork& Work : PreWorkTasks)
+	ParallelForWithPreWork( TEXT("CompressMipChain.PF"),AsyncCompressionTasks.Num(),1,
+		[&AsyncCompressionTasks](int32 TaskIndex)
 		{
-			bCompressionSucceeded = bCompressionSucceeded && TextureFormat->CompressImageEx(
-				&Work.SrcMip,
-				Work.MipIndex == FirstMipTailIndex ? CompressorCaps.NumMipsInTail : 1, // number of mips pointed to by SrcMip
-				Settings,
-				DebugTexturePathName,
-				bImageHasAlphaChannel,
-				CompressorCaps.ExtData,
-				Work.DestMip
-			);
-		}
-	}, EParallelForFlags::Unbalanced);
+			AsyncCompressionTasks[TaskIndex].DoWork();
+		},
+		[&PreWorkTasks, &TextureFormat, &OutMips, &bCompressionSucceeded, &CompressorCaps, &Settings, DebugTexturePathName, FirstMipTailIndex, bImageHasAlphaChannel]()
+		{
+			for (PreWork& Work : PreWorkTasks)
+			{
+				bCompressionSucceeded = bCompressionSucceeded && TextureFormat->CompressImageEx(
+					&Work.SrcMip,
+					Work.MipIndex == FirstMipTailIndex ? CompressorCaps.NumMipsInTail : 1, // number of mips pointed to by SrcMip
+					Settings,
+					DebugTexturePathName,
+					bImageHasAlphaChannel,
+					CompressorCaps.ExtData,
+					Work.DestMip
+				);
+			}
+		}, 
+		EParallelForFlags::Unbalanced);
 
 	for (int32 TaskIndex = 0; TaskIndex < AsyncCompressionTasks.Num(); ++TaskIndex)
 	{
@@ -2429,6 +2384,8 @@ public:
 		uint32& OutExtData
 	)
 	{
+		//TRACE_CPUPROFILER_EVENT_SCOPE(BuildTexture);
+
 		const ITextureFormat* TextureFormat = nullptr;
 
 		ITextureFormatManagerModule* TFM = GetTextureFormatManager();
@@ -2788,7 +2745,6 @@ private:
 				FTextureDownscaleSettings DownscaleSettings;
 				DownscaleSettings.Downscale = BuildSettings.Downscale;
 				DownscaleSettings.DownscaleOptions = BuildSettings.DownscaleOptions;
-				DownscaleSettings.bDitherMipMapAlpha = BuildSettings.bDitherMipMapAlpha;
 				DownscaleSettings.BlockSize = 4;
 		
 				DownscaleImage(*Mip, *Mip, DownscaleSettings);
diff --git a/Engine/Source/Developer/TextureCompressor/Public/TextureCompressorModule.h b/Engine/Source/Developer/TextureCompressor/Public/TextureCompressorModule.h
--- a/Engine/Source/Developer/TextureCompressor/Public/TextureCompressorModule.h
+++ b/Engine/Source/Developer/TextureCompressor/Public/TextureCompressorModule.h
@@ -122,8 +122,6 @@ struct FTextureBuildSettings
 	uint32 bForceNoAlphaChannel : 1;
 	/** Whether we should not discard the alpha channel when it contains 1 for the entire texture. */
 	uint32 bForceAlphaChannel : 1;
-	/** Whether the alpha channel should contain a dithered alpha value. */
-	uint32 bDitherMipMapAlpha : 1;
 	/** Whether bokeh alpha values should be computed for the texture. */
 	uint32 bComputeBokehAlpha : 1;
 	/** Whether the contents of the red channel should be replicated to all channels. */
@@ -239,7 +237,6 @@ struct FTextureBuildSettings
 		, bPreserveBorder(false)
 		, bForceNoAlphaChannel(false)
 		, bForceAlphaChannel(false)
-		, bDitherMipMapAlpha(false)
 		, bComputeBokehAlpha(false)
 		, bReplicateRed(false)
 		, bReplicateAlpha(false)
diff --git a/Engine/Source/Editor/UnrealEd/Classes/Factories/TextureFactory.h b/Engine/Source/Editor/UnrealEd/Classes/Factories/TextureFactory.h
--- a/Engine/Source/Editor/UnrealEd/Classes/Factories/TextureFactory.h
+++ b/Engine/Source/Editor/UnrealEd/Classes/Factories/TextureFactory.h
@@ -112,10 +112,6 @@ class UNREALED_API UTextureFactory : public UFactory, public IImportSettingsPars
 	UPROPERTY(EditAnywhere, Category=LODGroup, meta=(ToolTip="The group the texture belongs to"))
 	TEnumAsByte<enum TextureGroup> LODGroup;
 
-	/** If enabled, mip-map alpha values will be dithered for smooth transitions */
-	UPROPERTY(EditAnywhere, Category=DitherMipMaps, meta=(ToolTip="If enabled, mip-map alpha values will be dithered for smooth transitions"))
-	uint32 bDitherMipMapAlpha:1;
-	
 	/** Whether mip RGBA should be scaled to preserve the number of pixels with Value >= AlphaCoverageThresholds */
 	UPROPERTY(EditAnywhere, Category=PreserveAlphaCoverage, meta=(ToolTip="Whether mip RGBA should be scaled to preserve the number of pixels with Value >= AlphaCoverageThresholds"))
 	bool bDoScaleMipsForAlphaCoverage = false;
diff --git a/Engine/Source/Editor/UnrealEd/Private/Factories/EditorFactories.cpp b/Engine/Source/Editor/UnrealEd/Private/Factories/EditorFactories.cpp
--- a/Engine/Source/Editor/UnrealEd/Private/Factories/EditorFactories.cpp
+++ b/Engine/Source/Editor/UnrealEd/Private/Factories/EditorFactories.cpp
@@ -4225,7 +4225,6 @@ UObject* UTextureFactory::FactoryCreateBinary
 		ExistingNoAlpha = ExistingTexture->CompressionNoAlpha;
 		ExistingDeferCompression = ExistingTexture->DeferCompression;
 		ExistingFlipGreenChannel = ExistingTexture->bFlipGreenChannel;
-		ExistingDitherMipMapAlpha = ExistingTexture->bDitherMipMapAlpha;
 		ExistingbDoScaleMipsForAlphaCoverage = ExistingTexture->bDoScaleMipsForAlphaCoverage;
 		ExistingAlphaCoverageThresholds = ExistingTexture->AlphaCoverageThresholds;
 		ExistingAdjustBrightness = ExistingTexture->AdjustBrightness;
@@ -4346,7 +4345,6 @@ UObject* UTextureFactory::FactoryCreateBinary
 	Texture->CompressionNone		= NoCompression;
 	Texture->CompressionNoAlpha		= NoAlpha;
 	Texture->DeferCompression		= bDeferCompression;
-	Texture->bDitherMipMapAlpha		= bDitherMipMapAlpha;
 	Texture->bDoScaleMipsForAlphaCoverage = bDoScaleMipsForAlphaCoverage;
 	Texture->AlphaCoverageThresholds = AlphaCoverageThresholds;
 	
@@ -4409,7 +4407,6 @@ UObject* UTextureFactory::FactoryCreateBinary
 		Texture->CompressionNone = ExistingNoCompression;
 		Texture->CompressionNoAlpha = ExistingNoAlpha;
 		Texture->DeferCompression = ExistingDeferCompression;
-		Texture->bDitherMipMapAlpha = ExistingDitherMipMapAlpha;
 		Texture->bDoScaleMipsForAlphaCoverage = ExistingbDoScaleMipsForAlphaCoverage;
 		Texture->AlphaCoverageThresholds = ExistingAlphaCoverageThresholds;
 		Texture->bFlipGreenChannel = ExistingFlipGreenChannel;
diff --git a/Engine/Source/Runtime/Engine/Classes/Engine/Texture.h b/Engine/Source/Runtime/Engine/Classes/Engine/Texture.h
--- a/Engine/Source/Runtime/Engine/Classes/Engine/Texture.h
+++ b/Engine/Source/Runtime/Engine/Classes/Engine/Texture.h
@@ -1060,9 +1060,9 @@ public:
 	UPROPERTY(EditAnywhere, BlueprintReadWrite, Category=Compression, meta = (DisplayName = "ASTC Compression Quality"), AdvancedDisplay)
 	TEnumAsByte<enum ETextureCompressionQuality> CompressionQuality;
 
-	/** When true, the alpha channel of mip-maps and the base image are dithered for smooth LOD transitions. */
-	UPROPERTY(EditAnywhere, BlueprintReadWrite, Category=Texture, AdvancedDisplay)
-	uint32 bDitherMipMapAlpha:1;
+	/** Removed. */
+	UPROPERTY()
+	uint32 bDitherMipMapAlpha_DEPRECATED:1;
 
 	/** Whether mip RGBA should be scaled to preserve the number of pixels with Value >= AlphaCoverageThresholds.  AlphaCoverageThresholds are ignored if this is off. */
 	UPROPERTY(EditAnywhere, BlueprintReadWrite, Category=Texture, AdvancedDisplay)
diff --git a/Engine/Source/Runtime/Engine/Private/Texture.cpp b/Engine/Source/Runtime/Engine/Private/Texture.cpp
--- a/Engine/Source/Runtime/Engine/Private/Texture.cpp
+++ b/Engine/Source/Runtime/Engine/Private/Texture.cpp
@@ -2466,10 +2466,6 @@ FName GetDefaultTextureFormatName( const ITargetPlatform* TargetPlatform, const
 	{
 		TextureFormatName = NameDXT1;
 	}
-	else if (Texture->bDitherMipMapAlpha)
-	{
-		TextureFormatName = NameDXT5;
-	}
 	else
 	{
 		TextureFormatName = NameAutoDXT;
diff --git a/Engine/Source/Runtime/Engine/Private/TextureDerivedData.cpp b/Engine/Source/Runtime/Engine/Private/TextureDerivedData.cpp
--- a/Engine/Source/Runtime/Engine/Private/TextureDerivedData.cpp
+++ b/Engine/Source/Runtime/Engine/Private/TextureDerivedData.cpp
@@ -137,7 +137,11 @@ static void SerializeForKey(FArchive& Ar, const FTextureBuildSettings& Settings)
 	}
 
 	TempByte = Settings.bPreserveBorder; Ar << TempByte;
-	TempByte = Settings.bDitherMipMapAlpha; Ar << TempByte;
+
+	// bDitherMipMapAlpha was removed from Texture
+	//  serialize to DDC as if it was still around and false to keep keys the same:
+	uint8 bDitherMipMapAlpha = 0;
+	TempByte = bDitherMipMapAlpha; Ar << TempByte;
 
 	if (Settings.bDoScaleMipsForAlphaCoverage)
 	{
@@ -692,7 +696,6 @@ static void GetTextureBuildSettings(
 	OutBuildSettings.ColorAdjustment.AdjustMaxAlpha = Texture.AdjustMaxAlpha;
 	OutBuildSettings.bUseLegacyGamma = Texture.bUseLegacyGamma;
 	OutBuildSettings.bPreserveBorder = Texture.bPreserveBorder;
-	OutBuildSettings.bDitherMipMapAlpha = Texture.bDitherMipMapAlpha;
 
 	// in Texture , the fields bDoScaleMipsForAlphaCoverage and AlphaCoverageThresholds are independent
 	// but in the BuildSettings bDoScaleMipsForAlphaCoverage is only on if thresholds are valid (not all zero)
diff --git a/Engine/Source/Runtime/Engine/Private/TextureDerivedDataBuildUtils.cpp b/Engine/Source/Runtime/Engine/Private/TextureDerivedDataBuildUtils.cpp
--- a/Engine/Source/Runtime/Engine/Private/TextureDerivedDataBuildUtils.cpp
+++ b/Engine/Source/Runtime/Engine/Private/TextureDerivedDataBuildUtils.cpp
@@ -143,7 +143,6 @@ static void WriteBuildSettings(FCbWriter& Writer, const FTextureBuildSettings& B
 	WriteCbFieldWithDefault<bool>(Writer, "bPreserveBorder", BuildSettings.bPreserveBorder, DefaultSettings.bPreserveBorder);
 	WriteCbFieldWithDefault<bool>(Writer, "bForceNoAlphaChannel", BuildSettings.bForceNoAlphaChannel, DefaultSettings.bForceNoAlphaChannel);
 	WriteCbFieldWithDefault<bool>(Writer, "bForceAlphaChannel", BuildSettings.bForceAlphaChannel, DefaultSettings.bForceAlphaChannel);
-	WriteCbFieldWithDefault<bool>(Writer, "bDitherMipMapAlpha", BuildSettings.bDitherMipMapAlpha, DefaultSettings.bDitherMipMapAlpha);
 	WriteCbFieldWithDefault<bool>(Writer, "bComputeBokehAlpha", BuildSettings.bComputeBokehAlpha, DefaultSettings.bComputeBokehAlpha);
 	WriteCbFieldWithDefault<bool>(Writer, "bReplicateRed", BuildSettings.bReplicateRed, DefaultSettings.bReplicateRed);
 	WriteCbFieldWithDefault<bool>(Writer, "bReplicateAlpha", BuildSettings.bReplicateAlpha, DefaultSettings.bReplicateAlpha);
