Disable incorrect check in TextureCompressorModule
check( TableSize1D > 2 );
being hit due to strange composite normal path

diff --git a/Engine/Source/Developer/TextureCompressor/Private/TextureCompressorModule.cpp b/Engine/Source/Developer/TextureCompressor/Private/TextureCompressorModule.cpp
--- a/Engine/Source/Developer/TextureCompressor/Private/TextureCompressorModule.cpp
+++ b/Engine/Source/Developer/TextureCompressor/Private/TextureCompressorModule.cpp
@@ -136,7 +136,13 @@ public:
 		{
 			// blur only
 			// this is TMGS_Blur
-			check( TableSize1D > 2 );
+
+			// TMGS_Blur will always give us TableSize > 2
+			// @todo Oodle - temp check disabled
+			//   I  believe this should be true, but a bug in AssociatedNormalSourceMips
+			//   causes a 2 to come in here
+			//check( TableSize1D > 2 );
+
 			BuildGaussian1D(Table1D, TableSize1D, 1.0f, -SharpenFactor);
 			BuildFilterTable2DFrom1D(KernelWeights, Table1D, TableSize1D);
 			return;
@@ -153,6 +159,8 @@ public:
 		{
 			// 4x4 kernel with sharpen or blur: can alias a bit
 			// this is not used by standard TMGS_ mip options
+			//  one thing that can get you in here is GenerateTopMip
+			//	because it takes the standard 8 size and does /2
 			BuildFilterTable1DBase(Table1D, TableSize1D, 1.0f + SharpenFactor);
 			BuildFilterTable1DBase(NegativeTable1D, TableSize1D, -SharpenFactor);
 			BlurFilterTable1D(NegativeTable1D, TableSize1D, 1);
@@ -2078,7 +2086,7 @@ static float SpecularPowerToRoughness(float SpecularPower)
 }
 
 // @param CompositeTextureMode original type ECompositeTextureMode
-void ApplyCompositeTexture(FImage& RoughnessSourceMips, const FImage& NormalSourceMips, uint8 CompositeTextureMode, float CompositePower)
+static void ApplyCompositeTexture(FImage& RoughnessSourceMips, const FImage& NormalSourceMips, uint8 CompositeTextureMode, float CompositePower)
 {
 	check(RoughnessSourceMips.SizeX == NormalSourceMips.SizeX);
 	check(RoughnessSourceMips.SizeY == NormalSourceMips.SizeY);
@@ -2416,13 +2424,19 @@ public:
 		// apply roughness adjustment depending on normal map variation
 		if (AssociatedNormalSourceMips.Num())
 		{
-			//			check AssociatedNormalSourceMips.Format; 
+			// check AssociatedNormalSourceMips.Format; 
+			// ECompositeTextureMode is only NormalRoughness
+			//  composite texture should be a normal map
 
 			TArray<FImage> IntermediateAssociatedNormalSourceMipChain;
 
 			FTextureBuildSettings DefaultSettings;
 
 			// helps to reduce aliasing further
+			// @todo Oodle I think SharpenMipKernelSize=4 here is a bug
+			//	because GenerateTopMip does SharpenMipKernelSize/2
+			//	this actually gives you a kernel size of 2
+			//	which is a box filter
 			DefaultSettings.MipSharpening = -4.0f;
 			DefaultSettings.SharpenMipKernelSize = 4;
 			DefaultSettings.bApplyKernelToTopMip = true;
