fix filter settings used by GenerateTopMip for ApplyCompositeTexture
use an odd centered filter so there's no half pixel shift

diff --git a/Engine/Source/Developer/TextureCompressor/Private/TextureCompressorModule.cpp b/Engine/Source/Developer/TextureCompressor/Private/TextureCompressorModule.cpp
--- a/Engine/Source/Developer/TextureCompressor/Private/TextureCompressorModule.cpp
+++ b/Engine/Source/Developer/TextureCompressor/Private/TextureCompressorModule.cpp
@@ -131,30 +131,27 @@ public:
 		float NegativeTable1D[MaxKernelExtend];
 
 		FilterTableSize = TableSize1D;
-
-		if(SharpenFactor < 0.0f)
+		
+		if(TableSize1D == 2)
+		{
+			// 2x2 kernel: simple average
+			// SharpenFactor is ignored
+			// this is TMGS_SimpleAverage
+			KernelWeights[0] = KernelWeights[1] = KernelWeights[2] = KernelWeights[3] = 0.25f;
+			return;
+		}
+		else if(SharpenFactor < 0.0f)
 		{
 			// blur only
 			// this is TMGS_Blur
 
 			// TMGS_Blur will always give us TableSize > 2
-			// @todo CB - temp check disabled
-			//   I  believe this should be true, but a bug in AssociatedNormalSourceMips
-			//   causes a 2 to come in here
-			//check( TableSize1D > 2 );
+			check( TableSize1D > 2 );
 
 			BuildGaussian1D(Table1D, TableSize1D, 1.0f, -SharpenFactor);
 			BuildFilterTable2DFrom1D(KernelWeights, Table1D, TableSize1D);
 			return;
 		}
-		else if(TableSize1D == 2)
-		{
-			// 2x2 kernel: simple average
-			// SharpenFactor is ignored
-			// this is TMGS_SimpleAverage
-			KernelWeights[0] = KernelWeights[1] = KernelWeights[2] = KernelWeights[3] = 0.25f;
-			return;
-		}
 		else if(TableSize1D == 4)
 		{
 			// 4x4 kernel with sharpen or blur: can alias a bit
@@ -867,13 +864,23 @@ static void GenerateTopMip(const FImage& SrcImage, FImage& DestImage, const FTex
 
 	FImageKernel2D KernelDownsample;
 
-	//@todo CB this looks broken :
-	// /2 as input resolution is same as output resolution and the settings assumed the output is half resolution
-	KernelDownsample.BuildSeparatableGaussWithSharpen( FMath::Max( 2u, Settings.SharpenMipKernelSize / 2 ), Settings.MipSharpening );
-	
-	// use an odd filter here so it's not shifting :
-	//KernelDownsample.BuildSeparatableGaussWithSharpen( 5, -2.5f );
-	
+	if ( Settings.MipSharpening < 0.f )
+	{
+		// negative Sharpening is a Gaussian
+		//  this can make centered ("odd") filters, so the image doesn't shift
+		int32 OddMipKernelSize = Settings.SharpenMipKernelSize | 1;
+		KernelDownsample.BuildSeparatableGaussWithSharpen( OddMipKernelSize, Settings.MipSharpening );
+	}
+	else
+	{	
+		// non-Gaussians only support "even" filters
+		//	this causes a half-pixel shift of the top mip
+		// warn but then go ahead and do as requested
+		UE_LOG(LogTextureCompressor, Warning, TEXT("GenerateTopMip used with non-Gaussian blur filter will cause half pixel shift"));
+
+		KernelDownsample.BuildSeparatableGaussWithSharpen( Settings.SharpenMipKernelSize, Settings.MipSharpening );
+	}
+
 	DestImage.Init(SrcImage.SizeX, SrcImage.SizeY, SrcImage.NumSlices, SrcImage.Format, SrcImage.GammaSpace);
 
 	for (int32 SliceIndex = 0; SliceIndex < SrcImage.NumSlices; ++SliceIndex)
@@ -2452,20 +2459,20 @@ public:
 
 			FTextureBuildSettings DefaultSettings;
 
-			// helps to reduce aliasing further
-
-			//@todo CB this looks broken :			
-			DefaultSettings.MipSharpening = -4.0f;
-			DefaultSettings.SharpenMipKernelSize = 4;
-
-			// same as TMGS_Blur2 :
-			//DefaultSettings.MipSharpening = -4.0f;
-			//DefaultSettings.SharpenMipKernelSize = 6;
-
+			// apply a smooth Gaussian filter to the top level of the normal map
+			// the original comment says :
+			// "helps to reduce aliasing further"
+			DefaultSettings.MipSharpening = -3.5f;
+			DefaultSettings.SharpenMipKernelSize = 6;
 			DefaultSettings.bApplyKernelToTopMip = true;
+
 			// important to make accurate computation with normal length
 			DefaultSettings.bRenormalizeTopMip = true;
 
+			// @todo Oodle : filtering the normal map then computing roughness is fundamentally wrong
+			//  we should instead compute the roughness scalar first on the original normap map
+			//  then filter on the roughness scalar
+
 			if (!BuildTextureMips(AssociatedNormalSourceMips, DefaultSettings, CompressorCaps, IntermediateAssociatedNormalSourceMipChain))
 			{
 				UE_LOG(LogTextureCompressor, Warning, TEXT("Failed to generate texture mips for composite texture"));
diff --git a/Engine/Source/Runtime/Engine/Private/TextureDerivedData.cpp b/Engine/Source/Runtime/Engine/Private/TextureDerivedData.cpp
--- a/Engine/Source/Runtime/Engine/Private/TextureDerivedData.cpp
+++ b/Engine/Source/Runtime/Engine/Private/TextureDerivedData.cpp
@@ -245,7 +245,8 @@ void GetTextureDerivedDataKeySuffix(const UTexture& Texture, const FTextureBuild
 
 	if(IsValid(Texture.CompositeTexture) && Texture.CompositeTextureMode != CTM_Disabled)
 	{
-		CompositeTextureStr += TEXT("_");
+		// CompositeTextureMode output changed so force a new DDC key value :
+		CompositeTextureStr += TEXT("_C_");
 		CompositeTextureStr += Texture.CompositeTexture->Source.GetIdString();
 	}
 
