Fix alpha being turned on for textures when "Max Texture Size" is used
even when source image was opaque
this is caused by the mip filter being not quite normalized, which introduces a slightly less than 1 alpha
fix by increasing alpha tolerance to where a non-255 U8 alpha would actually be visible
prevents BC1 textures from becoming BC3 for no reason

diff --git a/Engine/Source/Developer/TextureCompressor/Private/TextureCompressorModule.cpp b/Engine/Source/Developer/TextureCompressor/Private/TextureCompressorModule.cpp
--- a/Engine/Source/Developer/TextureCompressor/Private/TextureCompressorModule.cpp
+++ b/Engine/Source/Developer/TextureCompressor/Private/TextureCompressorModule.cpp
@@ -2016,18 +2016,99 @@ static void FlipGreenChannel( FImage& Image )
 static bool DetectAlphaChannel(const FImage& InImage)
 {
 	TRACE_CPUPROFILER_EVENT_SCOPE(Texture.DetectAlphaChannel);
+	
+	// previous :
+	// "SMALL_NUMBER" is quite small, this provides almost no tolerance
+	// #define SMALL_NUMBER		(1.e-8f)
+	//const float FloatNonOpaqueAlpha = 1.0f - SMALL_NUMBER;
+	// 
+	// opaque alpha threshold where we'd quantize to < 255 in U8
+	//  images with only alpha larger than this are treated as opaque
+	const float FloatNonOpaqueAlpha = 254.5f / 255.f; // the U8 alpha threshold
+		
+	int64 NumPixels = (int64)InImage.SizeX * InImage.SizeY * InImage.NumSlices;
+
+	if ( InImage.Format == ERawImageFormat::BGRA8 )
+	{
+		TArrayView64<const FColor> SrcColorArray = InImage.AsBGRA8();
+		check( SrcColorArray.Num() == NumPixels );
+
+		const FColor * ColorPtr = &SrcColorArray[0];
+		const FColor * EndPtr = ColorPtr +  SrcColorArray.Num();
 
-	// Uncompressed data is required to check for an alpha channel.
-	const FLinearColor* SrcColors = (&InImage.AsRGBA32F()[0]);
-	const FLinearColor* LastColor = SrcColors + (InImage.SizeX * InImage.SizeY * InImage.NumSlices);
-	while (SrcColors < LastColor)
+		for(;ColorPtr<EndPtr;++ColorPtr)
+		{
+			if ( ColorPtr->A != 255 )
+			{
+				return true;
+			}
+		}
+	}
+	else if ( InImage.Format == ERawImageFormat::RGBA32F )
+	{
+		TArrayView64<const FLinearColor> SrcColorArray = InImage.AsRGBA32F();
+		check( SrcColorArray.Num() == NumPixels );
+
+		const FLinearColor * ColorPtr = &SrcColorArray[0];
+		const FLinearColor * EndPtr = ColorPtr +  SrcColorArray.Num();
+
+		for(;ColorPtr<EndPtr;++ColorPtr)
+		{
+			if (ColorPtr->A <= FloatNonOpaqueAlpha )
+			{
+				return true;
+			}
+		}
+	}
+	else if ( InImage.Format == ERawImageFormat::RGBA16 )
+	{
+		TArrayView64<const uint16> SrcChannelArray = InImage.AsRGBA16();
+		check( SrcChannelArray.Num() == NumPixels*4 );
+
+		const uint16 * ChannelPtr = &SrcChannelArray[0];
+		const uint16 * EndPtr = ChannelPtr +  SrcChannelArray.Num();
+
+		for(;ChannelPtr<EndPtr;ChannelPtr += 4)
+		{
+			if ( ChannelPtr[3] != 0xFFFF )
+			{
+				return true;
+			}
+		}
+	}
+	else if ( InImage.Format == ERawImageFormat::RGBA16F )
 	{
-		if (SrcColors->A < (1.0f - SMALL_NUMBER))
+		TArrayView64<const FFloat16Color> SrcColorArray = InImage.AsRGBA16F();
+		check( SrcColorArray.Num() == NumPixels );
+
+		const FFloat16Color * ColorPtr = &SrcColorArray[0];
+		const FFloat16Color * EndPtr = ColorPtr +  SrcColorArray.Num();
+		
+		for(;ColorPtr<EndPtr;++ColorPtr)
 		{
-			return true;
+			// 16F closest to 1.0 is 0.99951172
+			// use the float tolerance here? or check exactly ?
+			//if ( ColorPtr->A.GetFloat() < 1.f )
+			// use the same FloatNonOpaqueAlpha tolerance for consistency ?
+			if ( ColorPtr->A.GetFloat() < FloatNonOpaqueAlpha )
+			{
+				return true;
+			}
 		}
-		++SrcColors;
 	}
+	else if ( InImage.Format == ERawImageFormat::G8 ||
+		InImage.Format == ERawImageFormat::BGRE8 ||
+		InImage.Format == ERawImageFormat::G16 ||
+		InImage.Format == ERawImageFormat::R16F )
+	{
+		// source image formats don't have alpha
+	}
+	else
+	{
+		// new format ?
+		check(0);
+	}
+
 	return false;
 }
 
@@ -2252,6 +2333,7 @@ static bool CompressMipChain(
 	const ITextureFormat* TextureFormat,
 	const TArray<FImage>& MipChain,
 	const FTextureBuildSettings& Settings,
+	const bool bImageHasAlphaChannel,
 	FStringView DebugTexturePathName,
 	TArray<FCompressedImage2D>& OutMips,
 	uint32& OutNumMipsInTail,
@@ -2259,9 +2341,6 @@ static bool CompressMipChain(
 {
 	TRACE_CPUPROFILER_EVENT_SCOPE(Texture.CompressMipChain)
 
-	// @todo Oodle : DetectAlphaChannel could be merged with the CopyImage when the float promotion is done
-	const bool bImageHasAlphaChannel = !Settings.bForceNoAlphaChannel  && (Settings.bForceAlphaChannel || DetectAlphaChannel(MipChain[0]));
-
 	// now call the Ex version now that we have the proper MipChain
 	const FTextureFormatCompressorCaps CompressorCaps = TextureFormat->GetFormatCapabilitiesEx(Settings, MipChain.Num(), MipChain[0], bImageHasAlphaChannel);
 	OutNumMipsInTail = CompressorCaps.NumMipsInTail;
@@ -2437,6 +2516,30 @@ public:
 
 			return false;
 		}
+		
+		#if 0
+		// @todo Oodle : use or remove me
+		// experimental : detect if alpha is wanted in the mip filter process
+		bool bAlpha_Before = false;
+		bool bAlpha_Programmatic = false;
+		if ( ! BuildSettings.bForceNoAlphaChannel && ! BuildSettings.bForceAlphaChannel )
+		{
+			// at this point SourceMips is still in original format (eg. BGRA8 not yet promoted to linear float)
+			// bAlpha_Before can be accelerated from the source image in many cases (if it was 3-channel RGB source)
+			bAlpha_Before = DetectAlphaChannel(SourceMips[0]);
+
+			// can mip processing add alpha ?
+			//  does not gaurantee it definitely does add alpha, it just might
+			bAlpha_Programmatic =
+				BuildSettings.bComputeBokehAlpha ||
+				BuildSettings.bReplicateRed ||
+				BuildSettings.bChromaKeyTexture ||
+				( BuildSettings.CompositeTextureMode == CTM_NormalRoughnessToAlpha && AssociatedNormalSourceMips.Num() > 0 );
+			// (note bLongLatCubemap skips these processings, may be wrong)
+		}
+		// if bFilterAlpha is off, we can make mips in RGB 3-channel and skip A processing
+		bool bFilterAlpha = (! BuildSettings.bForceNoAlphaChannel) && ( BuildSettings.bForceAlphaChannel || bAlpha_Before || bAlpha_Programmatic );
+		#endif
 
 		TArray<FImage> IntermediateMipChain;
 
@@ -2484,6 +2587,21 @@ public:
 			}
 		}
 
+		
+		// DetectAlphaChannel on the top mip of the generated mip chain
+		//	it is now always Linear F32
+		//	BuildSettings could have programatically introduced alpha that was not in the source
+		// note the order of operations in bForceAlphaChannel and bForceNoAlphaChannel (ambiguity if both are on)
+		const bool bImageHasAlphaChannel = !BuildSettings.bForceNoAlphaChannel  && (BuildSettings.bForceAlphaChannel || DetectAlphaChannel(IntermediateMipChain[0]));
+	
+		#if 0
+		// @todo Oodle : experimental : verify bFilterAlpha was right
+		// if bAlpha_Programmatic, we have to recheck if alpha got made or not
+		check( bImageHasAlphaChannel == bAlpha_Before || bAlpha_Programmatic );
+		// if bFilterAlpha is off, we should not have alpha out
+		check( bFilterAlpha || !bImageHasAlphaChannel );
+		#endif
+
 		// Set the correct biased texture size so that the compressor understands the original source image size
 		// This is requires for platforms that may need to tile based on the original source texture size
 		BuildSettings.TopMipSize.X = IntermediateMipChain[0].SizeX;
@@ -2505,7 +2623,8 @@ public:
 			BuildSettings.ArraySlices = 1;
 		}
 		
-		return CompressMipChain(TextureFormat, IntermediateMipChain, BuildSettings, DebugTexturePathName, OutTextureMips, OutNumMipsInTail, OutExtData);
+		return CompressMipChain(TextureFormat, IntermediateMipChain, BuildSettings, bImageHasAlphaChannel, DebugTexturePathName,
+					OutTextureMips, OutNumMipsInTail, OutExtData);
 	}
 
 	// IModuleInterface implementation.
@@ -2702,7 +2821,7 @@ private:
 
 		TArray<FImage> GeneratedSourceMips;
 		if (bBuildSourceImage)
-		{			
+		{
 			// the source is larger than the compressor allows and no mip image exists to act as a smaller source.
 			// We must generate a suitable source image:
 			bool bSuitableFormat = PostOptionalUpscaleSourceMips.Last().Format == ERawImageFormat::RGBA32F;
@@ -2732,6 +2851,7 @@ private:
 				BaseImage.SizeY,
 				CompressorCaps.MaxTextureDimension
 				);
+			// Max Texture Size resizing happens here :
 			GenerateMipChain(BuildSettings, bSuitableFormat ? BaseImage : Temp, GeneratedSourceMips, LevelsToUsableSource);
 
 			check(GeneratedSourceMips.Num() != 0);
@@ -2757,6 +2877,8 @@ private:
 			{
 				// Generate the base mip from the long-lat source image.
 				GenerateBaseCubeMipFromLongitudeLatitude2D(Mip, Image, BuildSettings.MaxTextureResolution, BuildSettings.SourceEncodingOverride);
+	
+				// note the break here skips other adjustments:
 				break;
 			}
 			else
@@ -2806,6 +2928,7 @@ private:
 
 			// Apply color adjustments
 			AdjustImageColors(*Mip, BuildSettings);
+
 			if (BuildSettings.bComputeBokehAlpha)
 			{
 				// To get the occlusion in the BokehDOF shader working for all Bokeh textures.
