TextureCompressorModule NOP comments
investigate AdjustColors not being a nop with no RGB adjustments
HSV causes small drift which can be amplified by later processing (roughness)
HSV does bad things on values out of [0,1] range

diff --git a/Engine/Source/Developer/TextureCompressor/Private/TextureCompressorModule.cpp b/Engine/Source/Developer/TextureCompressor/Private/TextureCompressorModule.cpp
--- a/Engine/Source/Developer/TextureCompressor/Private/TextureCompressorModule.cpp
+++ b/Engine/Source/Developer/TextureCompressor/Private/TextureCompressorModule.cpp
@@ -1778,134 +1778,166 @@ void ITextureCompressorModule::GenerateAngularFilteredMips(TArray<FImage>& InOut
 	}
 }
 
-void ITextureCompressorModule::AdjustImageColors(FImage& Image, const FTextureBuildSettings& InBuildSettings)
+static inline void AdjustColors(FLinearColor * Colors,int64 Count,
+		const FTextureBuildSettings& InBuildSettings)
 {
-	const FColorAdjustmentParameters& InParams = InBuildSettings.ColorAdjustment;
-	check( Image.SizeX > 0 && Image.SizeY > 0 );
-
-	if( !FMath::IsNearlyEqual( InParams.AdjustBrightness, 1.0f, (float)KINDA_SMALL_NUMBER ) ||
-		!FMath::IsNearlyEqual( InParams.AdjustBrightnessCurve, 1.0f, (float)KINDA_SMALL_NUMBER ) ||
-		!FMath::IsNearlyEqual( InParams.AdjustSaturation, 1.0f, (float)KINDA_SMALL_NUMBER ) ||
-		!FMath::IsNearlyEqual( InParams.AdjustVibrance, 0.0f, (float)KINDA_SMALL_NUMBER ) ||
-		!FMath::IsNearlyEqual( InParams.AdjustRGBCurve, 1.0f, (float)KINDA_SMALL_NUMBER ) ||
-		!FMath::IsNearlyEqual( InParams.AdjustHue, 0.0f, (float)KINDA_SMALL_NUMBER ) ||
-		!FMath::IsNearlyEqual( InParams.AdjustMinAlpha, 0.0f, (float)KINDA_SMALL_NUMBER ) ||
-		!FMath::IsNearlyEqual( InParams.AdjustMaxAlpha, 1.0f, (float)KINDA_SMALL_NUMBER ) ||
-		InBuildSettings.bChromaKeyTexture )
+	for(int64 i=0;i<Count;i++)
 	{
-		TRACE_CPUPROFILER_EVENT_SCOPE(Texture.AdjustImageColors);
+		FLinearColor OriginalColor = Colors[i];
+	
+		#if 0
+		if ( ! InBuildSettings.bHDRSource )
+		{
+			// @todo Oodle: for non-HDR source ensure we are clamped as expected
+			//	(can drift out of clamp due to previous processing)
+			//  if you wind up even very slightly out of [0,1] range this function does bad things
+			OriginalColor.R = FMath::Clamp(OriginalColor.R, 0.0f, 1.f);
+			OriginalColor.G = FMath::Clamp(OriginalColor.G, 0.0f, 1.f);
+			OriginalColor.B = FMath::Clamp(OriginalColor.B, 0.0f, 1.f);
+		}
+		#endif
 
-		const FLinearColor ChromaKeyTarget = InBuildSettings.ChromaKeyColor;
+		const FLinearColor & ChromaKeyTarget = InBuildSettings.ChromaKeyColor;
 		const float ChromaKeyThreshold = InBuildSettings.ChromaKeyThreshold + SMALL_NUMBER;
-		const int64 NumPixels = (int64)Image.SizeX * Image.SizeY * Image.NumSlices;
-		TArrayView64<FLinearColor> ImageColors = Image.AsRGBA32F();
 
-		int64 NumPixelsEachJob;
-		int32 NumJobs = ImageParallelForComputeNumJobsForPixels(NumPixelsEachJob,NumPixels);
-
-		//TFunction<void (int32)> 
-		auto AdjustImageColorsFunc = [&](int32 Index)
+		if (InBuildSettings.bChromaKeyTexture && (OriginalColor.Equals(ChromaKeyTarget, ChromaKeyThreshold)))
 		{
-			int64 StartIndex = Index * NumPixelsEachJob;
-			int64 EndIndex = FMath::Min(StartIndex + NumPixelsEachJob, NumPixels);
-			for (int64 CurPixelIndex = StartIndex; CurPixelIndex < EndIndex; ++CurPixelIndex)
-			{
-				const FLinearColor OriginalColorRaw = ImageColors[CurPixelIndex];
+			OriginalColor = FLinearColor::Transparent;
 
-				FLinearColor OriginalColor = OriginalColorRaw;
-				if (InBuildSettings.bChromaKeyTexture && (OriginalColor.Equals(ChromaKeyTarget, ChromaKeyThreshold)))
-				{
-					OriginalColor = FLinearColor::Transparent;
-				}
+			//@todo Oodle: strange: no return? processing continues on the transparent color...
+			//	  this was likely unintentional
+		}
 
-				// Convert to HSV
-				FLinearColor HSVColor = OriginalColor.LinearRGBToHSV();
-				float& PixelHue = HSVColor.R;
-				float& PixelSaturation = HSVColor.G;
-				float& PixelValue = HSVColor.B;
+		// NOTE: if OriginalColor has HDR/floats in it, this does not handle it well
+		//	it implicitly discards negatives (and negatives cause color shifts)
+		//	for values > 1 the clamp behavior is very strange
 
-				float OriginalLuminance = PixelValue;
+		// Convert to HSV
+		FLinearColor HSVColor = OriginalColor.LinearRGBToHSV();
+		float& PixelHue = HSVColor.R;
+		float& PixelSaturation = HSVColor.G;
+		float& PixelValue = HSVColor.B;
 
-				// Apply brightness adjustment
-				PixelValue *= InParams.AdjustBrightness;
+		float OriginalLuminance = PixelValue;
 
-				// Apply brightness power adjustment
-				if (!FMath::IsNearlyEqual(InParams.AdjustBrightnessCurve, 1.0f, (float)KINDA_SMALL_NUMBER) && InParams.AdjustBrightnessCurve != 0.0f)
-				{
-					// Raise HSV.V to the specified power
-					PixelValue = FMath::Pow(PixelValue, InParams.AdjustBrightnessCurve);
-				}
+		const FColorAdjustmentParameters& InParams = InBuildSettings.ColorAdjustment;
 
-				// Apply "vibrance" adjustment
-				if (!FMath::IsNearlyZero(InParams.AdjustVibrance, (float)KINDA_SMALL_NUMBER))
-				{
-					const float SatRaisePow = 5.0f;
-					const float InvSatRaised = FMath::Pow(1.0f - PixelSaturation, SatRaisePow);
+		// Apply brightness adjustment
+		PixelValue *= InParams.AdjustBrightness;
 
-					const float ClampedVibrance = FMath::Clamp(InParams.AdjustVibrance, 0.0f, 1.0f);
-					const float HalfVibrance = ClampedVibrance * 0.5f;
+		// Apply brightness power adjustment
+		if (!FMath::IsNearlyEqual(InParams.AdjustBrightnessCurve, 1.0f, (float)KINDA_SMALL_NUMBER) && InParams.AdjustBrightnessCurve != 0.0f)
+		{
+			// Raise HSV.V to the specified power
+			PixelValue = FMath::Pow(PixelValue, InParams.AdjustBrightnessCurve);
+		}
 
-					const float SatProduct = HalfVibrance * InvSatRaised;
+		// Apply "vibrance" adjustment
+		if (!FMath::IsNearlyZero(InParams.AdjustVibrance, (float)KINDA_SMALL_NUMBER))
+		{
+			const float SatRaisePow = 5.0f;
+			const float InvSatRaised = FMath::Pow(1.0f - PixelSaturation, SatRaisePow);
 
-					PixelSaturation += SatProduct;
-				}
+			const float ClampedVibrance = FMath::Clamp(InParams.AdjustVibrance, 0.0f, 1.0f);
+			const float HalfVibrance = ClampedVibrance * 0.5f;
 
-				// Apply saturation adjustment
-				PixelSaturation *= InParams.AdjustSaturation;
+			const float SatProduct = HalfVibrance * InvSatRaised;
 
-				// Apply hue adjustment
-				PixelHue += InParams.AdjustHue;
+			PixelSaturation += SatProduct;
+		}
 
-				// Clamp HSV values
-				{
-					PixelHue = FMath::Fmod(PixelHue, 360.0f);
-					if (PixelHue < 0.0f)
-					{
-						// Keep the hue value positive as HSVToLinearRGB prefers that
-						PixelHue += 360.0f;
-					}
-					PixelSaturation = FMath::Clamp(PixelSaturation, 0.0f, 1.0f);
+		// Apply saturation adjustment
+		PixelSaturation *= InParams.AdjustSaturation;
 
-					// Clamp brightness if non-HDR
-					if (!InBuildSettings.bHDRSource)
-					{
-						PixelValue = FMath::Clamp(PixelValue, 0.0f, 1.0f);
-					}
-				}
+		// Apply hue adjustment
+		PixelHue += InParams.AdjustHue;
 
-				// Convert back to a linear color
-				FLinearColor LinearColor = HSVColor.HSVToLinearRGB();
+		// Clamp HSV values
+		PixelHue = FMath::Fmod(PixelHue, 360.0f);
+		if (PixelHue < 0.0f)
+		{
+			// Keep the hue value positive as HSVToLinearRGB prefers that
+			PixelHue += 360.0f;
+		}
+		PixelSaturation = FMath::Clamp(PixelSaturation, 0.0f, 1.0f);
 
-				// Apply RGB curve adjustment (linear space)
-				if (!FMath::IsNearlyEqual(InParams.AdjustRGBCurve, 1.0f, (float)KINDA_SMALL_NUMBER) && InParams.AdjustRGBCurve != 0.0f)
-				{
-					LinearColor.R = FMath::Pow(LinearColor.R, InParams.AdjustRGBCurve);
-					LinearColor.G = FMath::Pow(LinearColor.G, InParams.AdjustRGBCurve);
-					LinearColor.B = FMath::Pow(LinearColor.B, InParams.AdjustRGBCurve);
-				}
+		// Clamp brightness if non-HDR
+		if (!InBuildSettings.bHDRSource)
+		{
+			PixelValue = FMath::Clamp(PixelValue, 0.0f, 1.0f);
+		}
 
-				// Clamp HDR RGB channels to 1 or the original luminance (max original RGB channel value), whichever is greater
-				if (InBuildSettings.bHDRSource)
-				{
-					LinearColor.R = FMath::Clamp(LinearColor.R, 0.0f, (OriginalLuminance > 1.0f ? OriginalLuminance : 1.0f));
-					LinearColor.G = FMath::Clamp(LinearColor.G, 0.0f, (OriginalLuminance > 1.0f ? OriginalLuminance : 1.0f));
-					LinearColor.B = FMath::Clamp(LinearColor.B, 0.0f, (OriginalLuminance > 1.0f ? OriginalLuminance : 1.0f));
-				}
+		// Convert back to a linear color
+		FLinearColor LinearColor = HSVColor.HSVToLinearRGB();
 
-				// Remap the alpha channel
-				LinearColor.A = FMath::Lerp(InParams.AdjustMinAlpha, InParams.AdjustMaxAlpha, OriginalColor.A);
-				ImageColors[CurPixelIndex] = LinearColor;
-			}
-		};
+		// Apply RGB curve adjustment (linear space)
+		if (!FMath::IsNearlyEqual(InParams.AdjustRGBCurve, 1.0f, (float)KINDA_SMALL_NUMBER) && InParams.AdjustRGBCurve != 0.0f)
+		{
+			LinearColor.R = FMath::Pow(LinearColor.R, InParams.AdjustRGBCurve);
+			LinearColor.G = FMath::Pow(LinearColor.G, InParams.AdjustRGBCurve);
+			LinearColor.B = FMath::Pow(LinearColor.B, InParams.AdjustRGBCurve);
+		}
+
+		// Clamp HDR RGB channels to 1 or the original luminance (max original RGB channel value), whichever is greater
+		// @todo Oodle: this is a very odd thing to do
+		//		clamping at OriginalLuminance if you do AdjustBrightness or AdjustBrightnessCurve ?
+		//	    that would keep values brighter than 1.f unchanged, but bring up lower ones to 1.f
+		//	  I question whether this should just be completely removed
+		if (InBuildSettings.bHDRSource)
+		{
+			LinearColor.R = FMath::Clamp(LinearColor.R, 0.0f, (OriginalLuminance > 1.0f ? OriginalLuminance : 1.0f));
+			LinearColor.G = FMath::Clamp(LinearColor.G, 0.0f, (OriginalLuminance > 1.0f ? OriginalLuminance : 1.0f));
+			LinearColor.B = FMath::Clamp(LinearColor.B, 0.0f, (OriginalLuminance > 1.0f ? OriginalLuminance : 1.0f));
+		}
+
+		// Remap the alpha channel
+		// @todo Oodle: clamp Original A in [0,1] ?
+		LinearColor.A = FMath::Lerp(InParams.AdjustMinAlpha, InParams.AdjustMaxAlpha, OriginalColor.A);
+
+		Colors[i] = LinearColor;
+	}
+}
+
+void ITextureCompressorModule::AdjustImageColors(FImage& Image, const FTextureBuildSettings& InBuildSettings)
+{
+	const FColorAdjustmentParameters& InParams = InBuildSettings.ColorAdjustment;
+	check( Image.SizeX > 0 && Image.SizeY > 0 );
+
+	bool bAdjustNeeded =
+		!FMath::IsNearlyEqual( InParams.AdjustBrightness, 1.0f, (float)KINDA_SMALL_NUMBER ) ||
+		!FMath::IsNearlyEqual( InParams.AdjustBrightnessCurve, 1.0f, (float)KINDA_SMALL_NUMBER ) ||
+		!FMath::IsNearlyEqual( InParams.AdjustSaturation, 1.0f, (float)KINDA_SMALL_NUMBER ) ||
+		!FMath::IsNearlyEqual( InParams.AdjustVibrance, 0.0f, (float)KINDA_SMALL_NUMBER ) ||
+		!FMath::IsNearlyEqual( InParams.AdjustRGBCurve, 1.0f, (float)KINDA_SMALL_NUMBER ) ||
+		!FMath::IsNearlyEqual( InParams.AdjustHue, 0.0f, (float)KINDA_SMALL_NUMBER ) ||
+		!FMath::IsNearlyEqual( InParams.AdjustMinAlpha, 0.0f, (float)KINDA_SMALL_NUMBER ) ||
+		!FMath::IsNearlyEqual( InParams.AdjustMaxAlpha, 1.0f, (float)KINDA_SMALL_NUMBER ) ||
+		InBuildSettings.bChromaKeyTexture;
 		
+	if ( bAdjustNeeded )
+	{
+		TRACE_CPUPROFILER_EVENT_SCOPE(Texture.AdjustImageColors);
+
+		const int64 NumPixels = (int64)Image.SizeX * Image.SizeY * Image.NumSlices;
+		TArrayView64<FLinearColor> ImageColors = Image.AsRGBA32F();
+
+		int64 NumPixelsEachJob;
+		int32 NumJobs = ImageParallelForComputeNumJobsForPixels(NumPixelsEachJob,NumPixels);
+
 		// bForceSingleThread is set to true when: 
 		// editor or cooker is loading as this is when the derived data cache is rebuilt as it will already be limited to a single thread 
 		//     and thus overhead of multithreading will simply make it slower
 		// @todo Oodle - this is done here and not in other similar ParallelFor places here.  It should either be done everywhere or nowhere.
 		bool bForceSingleThread = GIsEditorLoadingPackage || GIsCookerLoadingPackage || IsInAsyncLoadingThread();
 
-		ParallelFor( TEXT("Texture.AdjustImageColorsFunc.PF"),NumJobs,1, AdjustImageColorsFunc, 
-			(bForceSingleThread ? EParallelForFlags::ForceSingleThread : EParallelForFlags::None) );
+		ParallelFor( TEXT("Texture.AdjustImageColorsFunc.PF"),NumJobs,1, [&](int32 Index)
+		{
+			int64 StartIndex = Index * NumPixelsEachJob;
+			int64 EndIndex = FMath::Min(StartIndex + NumPixelsEachJob, NumPixels);
+
+			AdjustColors(&ImageColors[StartIndex],EndIndex-StartIndex,InBuildSettings);
+		}
+		, (bForceSingleThread ? EParallelForFlags::ForceSingleThread : EParallelForFlags::None) );
 	}
 }
 
@@ -2517,6 +2549,11 @@ public:
 			return false;
 		}
 		
+		// @todo Oodle: option to dump the Source image here
+		//		we have dump in TextureFormatOodle for the after-processing (before encoding) image
+		//		get a dump spot for before-processing as well
+
+
 		#if 0
 		// @todo Oodle : use or remove me
 		// experimental : detect if alpha is wanted in the mip filter process
