TextureCompressorModule AdjustColors optimize a little
verified nop

diff --git a/Engine/Source/Developer/TextureCompressor/Private/TextureCompressorModule.cpp b/Engine/Source/Developer/TextureCompressor/Private/TextureCompressorModule.cpp
--- a/Engine/Source/Developer/TextureCompressor/Private/TextureCompressorModule.cpp
+++ b/Engine/Source/Developer/TextureCompressor/Private/TextureCompressorModule.cpp
@@ -1781,31 +1781,62 @@ void ITextureCompressorModule::GenerateAngularFilteredMips(TArray<FImage>& InOut
 static inline void AdjustColors(FLinearColor * Colors,int64 Count,
 		const FTextureBuildSettings& InBuildSettings)
 {
+	const FColorAdjustmentParameters& Params = InBuildSettings.ColorAdjustment;
+	
+	// @todo Oodle : not the same checks as bAdjustNeeded outside
+	//	 but preserves legacy behavior
+	bool bAdjustBrightnessCurve = (!FMath::IsNearlyEqual(Params.AdjustBrightnessCurve, 1.0f, (float)KINDA_SMALL_NUMBER) && Params.AdjustBrightnessCurve != 0.0f);
+	bool bAdjustVibrance = (!FMath::IsNearlyZero(Params.AdjustVibrance, (float)KINDA_SMALL_NUMBER));
+	bool bAdjustRGBCurve = (!FMath::IsNearlyEqual(Params.AdjustRGBCurve, 1.0f, (float)KINDA_SMALL_NUMBER) && Params.AdjustRGBCurve != 0.0f);
+	bool bAdjustSaturation = ( Params.AdjustSaturation != 1.f || bAdjustVibrance );
+	bool bAdjustValue = ( Params.AdjustBrightness != 1.f ) || bAdjustBrightnessCurve;
+
+	float AdjustHue = Params.AdjustHue;
+	if ( AdjustHue != 0.f )
+	{
+		// Params.AdjustHue should be in [0,360] , make sure
+		if ( AdjustHue < 0.f || AdjustHue > 360.f )
+		{
+			AdjustHue = fmodf(AdjustHue, 360.0f);
+			if ( AdjustHue < 0.f )
+			{
+				AdjustHue += 360.f;
+			}
+		}
+	}
+	
+	// BuildSettings.ChromaKeyColor is an FColor
+	FLinearColor ChromaKeyColor(InBuildSettings.ChromaKeyColor);
+
+	bool bHDRSource = InBuildSettings.bHDRSource;
+	bool bChromaKeyTexture = InBuildSettings.bChromaKeyTexture;
+	float ChromaKeyThreshold = InBuildSettings.ChromaKeyThreshold + SMALL_NUMBER;
+
 	for(int64 i=0;i<Count;i++)
 	{
 		FLinearColor OriginalColor = Colors[i];
 	
 		#if 0
-		if ( ! InBuildSettings.bHDRSource )
+		if ( ! bHDRSource )
 		{
 			// @todo Oodle: for non-HDR source ensure we are clamped as expected
 			//	(can drift out of clamp due to previous processing)
 			//  if you wind up even very slightly out of [0,1] range this function does bad things
+			//	I think this probably should be done, but left disabled to preserve legacy behavior
 			OriginalColor.R = FMath::Clamp(OriginalColor.R, 0.0f, 1.f);
 			OriginalColor.G = FMath::Clamp(OriginalColor.G, 0.0f, 1.f);
 			OriginalColor.B = FMath::Clamp(OriginalColor.B, 0.0f, 1.f);
 		}
 		#endif
 
-		const FLinearColor & ChromaKeyTarget = InBuildSettings.ChromaKeyColor;
-		const float ChromaKeyThreshold = InBuildSettings.ChromaKeyThreshold + SMALL_NUMBER;
-
-		if (InBuildSettings.bChromaKeyTexture && (OriginalColor.Equals(ChromaKeyTarget, ChromaKeyThreshold)))
+		if (bChromaKeyTexture && (OriginalColor.Equals(ChromaKeyColor, ChromaKeyThreshold)))
 		{
 			OriginalColor = FLinearColor::Transparent;
 
 			//@todo Oodle: strange: no return? processing continues on the transparent color...
 			//	  this was likely unintentional
+			//Colors[i] = FLinearColor::Transparent;
+			//continue;
 		}
 
 		// NOTE: if OriginalColor has HDR/floats in it, this does not handle it well
@@ -1820,62 +1851,74 @@ static inline void AdjustColors(FLinearColor * Colors,int64 Count,
 
 		float OriginalLuminance = PixelValue;
 
-		const FColorAdjustmentParameters& InParams = InBuildSettings.ColorAdjustment;
-
-		// Apply brightness adjustment
-		PixelValue *= InParams.AdjustBrightness;
-
-		// Apply brightness power adjustment
-		if (!FMath::IsNearlyEqual(InParams.AdjustBrightnessCurve, 1.0f, (float)KINDA_SMALL_NUMBER) && InParams.AdjustBrightnessCurve != 0.0f)
+		if ( bAdjustValue )
 		{
-			// Raise HSV.V to the specified power
-			PixelValue = FMath::Pow(PixelValue, InParams.AdjustBrightnessCurve);
-		}
+			// Apply brightness adjustment
+			PixelValue *= Params.AdjustBrightness;
 
-		// Apply "vibrance" adjustment
-		if (!FMath::IsNearlyZero(InParams.AdjustVibrance, (float)KINDA_SMALL_NUMBER))
+			// Apply brightness power adjustment
+			if ( bAdjustBrightnessCurve )
+			{
+				// Raise HSV.V to the specified power
+				PixelValue = FMath::Pow(PixelValue, Params.AdjustBrightnessCurve);
+			}
+		
+			// Clamp brightness if non-HDR
+			if (!bHDRSource)
+			{
+				PixelValue = FMath::Clamp(PixelValue, 0.0f, 1.0f);
+			}
+		}
+		
+		if ( bAdjustSaturation )
 		{
-			const float SatRaisePow = 5.0f;
-			const float InvSatRaised = FMath::Pow(1.0f - PixelSaturation, SatRaisePow);
-
-			const float ClampedVibrance = FMath::Clamp(InParams.AdjustVibrance, 0.0f, 1.0f);
-			const float HalfVibrance = ClampedVibrance * 0.5f;
+			// PixelSaturation is >= 0 but not <= 1
+			//  because negative RGB can come into this function which gives Saturation > 1
 
-			const float SatProduct = HalfVibrance * InvSatRaised;
+			// Apply "vibrance" adjustment
+			if ( bAdjustVibrance )
+			{
+				// note: AdjustVibrance is disabled for HDR source in the Texture UPROPERTIES
+				//    (unclear why, this is no worse than anything else here on HDR)
+				// @todo Oodle : we'd like to just do the multiplies to make pow 5 but can't because it could change floats
+				const float SatRaisePow = 5.0f;
+				const float InvSatRaised = FMath::Pow(1.0f - PixelSaturation, SatRaisePow);
 
-			PixelSaturation += SatProduct;
-		}
+				const float ClampedVibrance = FMath::Clamp(Params.AdjustVibrance, 0.0f, 1.0f);
+				const float HalfVibrance = ClampedVibrance * 0.5f;
 
-		// Apply saturation adjustment
-		PixelSaturation *= InParams.AdjustSaturation;
+				const float SatProduct = HalfVibrance * InvSatRaised;
 
-		// Apply hue adjustment
-		PixelHue += InParams.AdjustHue;
+				PixelSaturation += SatProduct;
+			}
 
-		// Clamp HSV values
-		PixelHue = FMath::Fmod(PixelHue, 360.0f);
-		if (PixelHue < 0.0f)
-		{
-			// Keep the hue value positive as HSVToLinearRGB prefers that
-			PixelHue += 360.0f;
+			// Apply saturation adjustment
+			PixelSaturation *= Params.AdjustSaturation;
+			PixelSaturation = FMath::Clamp(PixelSaturation, 0.0f, 1.0f);
 		}
-		PixelSaturation = FMath::Clamp(PixelSaturation, 0.0f, 1.0f);
 
-		// Clamp brightness if non-HDR
-		if (!InBuildSettings.bHDRSource)
+		// Apply hue adjustment
+		if ( AdjustHue != 0.f )
 		{
-			PixelValue = FMath::Clamp(PixelValue, 0.0f, 1.0f);
+			// PixelHue is [0,360) but AdjustHue is [0,360]
+			PixelHue += AdjustHue;
+
+			// Clamp HSV values
+			if ( PixelHue >= 360.f )
+			{
+				PixelHue -= 360.f;
+			}
 		}
 
 		// Convert back to a linear color
 		FLinearColor LinearColor = HSVColor.HSVToLinearRGB();
 
 		// Apply RGB curve adjustment (linear space)
-		if (!FMath::IsNearlyEqual(InParams.AdjustRGBCurve, 1.0f, (float)KINDA_SMALL_NUMBER) && InParams.AdjustRGBCurve != 0.0f)
+		if ( bAdjustRGBCurve )
 		{
-			LinearColor.R = FMath::Pow(LinearColor.R, InParams.AdjustRGBCurve);
-			LinearColor.G = FMath::Pow(LinearColor.G, InParams.AdjustRGBCurve);
-			LinearColor.B = FMath::Pow(LinearColor.B, InParams.AdjustRGBCurve);
+			LinearColor.R = FMath::Pow(LinearColor.R, Params.AdjustRGBCurve);
+			LinearColor.G = FMath::Pow(LinearColor.G, Params.AdjustRGBCurve);
+			LinearColor.B = FMath::Pow(LinearColor.B, Params.AdjustRGBCurve);
 		}
 
 		// Clamp HDR RGB channels to 1 or the original luminance (max original RGB channel value), whichever is greater
@@ -1883,7 +1926,7 @@ static inline void AdjustColors(FLinearColor * Colors,int64 Count,
 		//		clamping at OriginalLuminance if you do AdjustBrightness or AdjustBrightnessCurve ?
 		//	    that would keep values brighter than 1.f unchanged, but bring up lower ones to 1.f
 		//	  I question whether this should just be completely removed
-		if (InBuildSettings.bHDRSource)
+		if (bHDRSource)
 		{
 			LinearColor.R = FMath::Clamp(LinearColor.R, 0.0f, (OriginalLuminance > 1.0f ? OriginalLuminance : 1.0f));
 			LinearColor.G = FMath::Clamp(LinearColor.G, 0.0f, (OriginalLuminance > 1.0f ? OriginalLuminance : 1.0f));
@@ -1892,7 +1935,7 @@ static inline void AdjustColors(FLinearColor * Colors,int64 Count,
 
 		// Remap the alpha channel
 		// @todo Oodle: clamp Original A in [0,1] ?
-		LinearColor.A = FMath::Lerp(InParams.AdjustMinAlpha, InParams.AdjustMaxAlpha, OriginalColor.A);
+		LinearColor.A = FMath::Lerp(Params.AdjustMinAlpha, Params.AdjustMaxAlpha, OriginalColor.A);
 
 		Colors[i] = LinearColor;
 	}
@@ -1903,6 +1946,11 @@ void ITextureCompressorModule::AdjustImageColors(FImage& Image, const FTextureBu
 	const FColorAdjustmentParameters& InParams = InBuildSettings.ColorAdjustment;
 	check( Image.SizeX > 0 && Image.SizeY > 0 );
 
+	// @todo Oodle : this bAdjustNeeded is not checking the same conditions to enable these adjustments
+	//		as is used inside the AdjustColors() routine
+	//		this is how it was done in the past, so keep it the same to preserve legacy operation
+	//	if possible in the future factor this Needed check out so it is shared code
+
 	bool bAdjustNeeded =
 		!FMath::IsNearlyEqual( InParams.AdjustBrightness, 1.0f, (float)KINDA_SMALL_NUMBER ) ||
 		!FMath::IsNearlyEqual( InParams.AdjustBrightnessCurve, 1.0f, (float)KINDA_SMALL_NUMBER ) ||
@@ -1935,7 +1983,10 @@ void ITextureCompressorModule::AdjustImageColors(FImage& Image, const FTextureBu
 			int64 StartIndex = Index * NumPixelsEachJob;
 			int64 EndIndex = FMath::Min(StartIndex + NumPixelsEachJob, NumPixels);
 
-			AdjustColors(&ImageColors[StartIndex],EndIndex-StartIndex,InBuildSettings);
+			FLinearColor * First = &ImageColors[StartIndex];
+			int64 Count = EndIndex-StartIndex;
+
+			AdjustColors(First,Count,InBuildSettings);
 		}
 		, (bForceSingleThread ? EParallelForFlags::ForceSingleThread : EParallelForFlags::None) );
 	}
