Faster texture mip generation.
Make inner loops to use vector instructions and avoid per pixel conditional jumps.
Unrolls kernel inner loop for 8x8 case.

diff --git a/Engine/Source/Developer/TextureCompressor/Private/TextureCompressorModule.cpp b/Engine/Source/Developer/TextureCompressor/Private/TextureCompressorModule.cpp
--- a/Engine/Source/Developer/TextureCompressor/Private/TextureCompressorModule.cpp
+++ b/Engine/Source/Developer/TextureCompressor/Private/TextureCompressorModule.cpp
@@ -77,7 +77,7 @@ struct FImageView2D
 // 2D sample lookup with input conversion
 // requires SourceImageData.SizeX and SourceImageData.SizeY to be power of two
 template <EMipGenAddressMode AddressMode>
-FLinearColor LookupSourceMip(const FImageView2D& SourceImageData, int32 X, int32 Y)
+static const FLinearColor& LookupSourceMip(const FImageView2D& SourceImageData, int32 X, int32 Y)
 {
 	if(AddressMode == MGTAM_Wrap)
 	{
@@ -97,14 +97,14 @@ FLinearColor LookupSourceMip(const FImageView2D& SourceImageData, int32 X, int32
 		if((uint32)X >= (uint32)SourceImageData.SizeX
 			|| (uint32)Y >= (uint32)SourceImageData.SizeY)
 		{
-			return FLinearColor(0, 0, 0, 0);
+			static FLinearColor Black(0, 0, 0, 0);
+			return Black;
 		}
 	}
 	else
 	{
 		check(0);
 	}
-	//return *(SourceImageData.AsRGBA32F() + X + Y * SourceImageData.SizeX);
 	return SourceImageData.Access(X,Y);
 }
 
@@ -555,15 +555,176 @@ static void GenerateMip2x2Simple(
 		int32 EndIndex = FMath::Min(StartIndex + NumRowsEachJob, DestImageData.SizeY);
 		for (int32 DestY = StartIndex; DestY < EndIndex; ++DestY)
 		{
-			FLinearColor * DestRow = &DestImageData.Access(0,DestY);
-			const FLinearColor * SourceRow0 = &SourceImageData.Access(0,2*DestY);
-			const FLinearColor * SourceRow1 = &SourceImageData.Access(0,2*DestY+1);
+			float* DestRow = (float*)&DestImageData.Access(0, DestY);
+			const float* SourceRow0 = (const float*)&SourceImageData.Access(0, 2*DestY);
+			const float* SourceRow1 = (const float*)&SourceImageData.Access(0, 2*DestY+1);
 
-			for ( int32 DestX = 0;DestX < DestImageData.SizeX; DestX++ )
+			const VectorRegister4Float Mul = VectorSetFloat1(0.25f);
+			for (int32 DestX = 0; DestX < DestImageData.SizeX; DestX++)
 			{
-				DestRow[DestX] = (SourceRow0[0] + SourceRow0[1] + SourceRow1[0] + SourceRow1[1]) * 0.25f;
-				SourceRow0 += 2;
-				SourceRow1 += 2;
+				VectorRegister4Float A = VectorLoad(&SourceRow0[0]);
+				VectorRegister4Float B = VectorLoad(&SourceRow0[4]);
+				VectorRegister4Float C = VectorLoad(&SourceRow1[0]);
+				VectorRegister4Float D = VectorLoad(&SourceRow1[4]);
+				VectorRegister4Float Sum = VectorAdd(VectorAdd(VectorAdd(A, B), C), D);
+				VectorRegister4Float Avg = VectorMultiply(Sum, Mul);
+				VectorStore(Avg, &DestRow[0]);
+				SourceRow0 += 8;
+				SourceRow1 += 8;
+				DestRow += 4;
+			}
+		}
+	});
+}
+
+template <EMipGenAddressMode AddressMode>
+static void GenerateMipUnfiltered(const FImageView2D& SourceImageData, FImageView2D& DestImageData, const FVector4f AlphaScale, uint32 ScaleFactor)
+{
+	int32 NumRowsEachJob;
+	int32 NumJobs = ImageParallelForComputeNumJobsForRows(NumRowsEachJob, DestImageData.SizeX, DestImageData.SizeY);
+
+	ParallelFor(TEXT("Texture.GenerateMipUnfiltered.PF"), NumJobs, 1, [&](int32 Index)
+	{
+		int32 StartIndex = Index * NumRowsEachJob;
+		int32 EndIndex = FMath::Min(StartIndex + NumRowsEachJob, DestImageData.SizeY);
+		for (int32 DestY = StartIndex; DestY < EndIndex; ++DestY)
+		{
+			for (int32 DestX = 0; DestX < DestImageData.SizeX; DestX++)
+			{
+				const int32 SourceX = DestX * ScaleFactor;
+				const int32 SourceY = DestY * ScaleFactor;
+				const FLinearColor& Sample = LookupSourceMip<AddressMode>(SourceImageData, SourceX, SourceY);
+				VectorRegister4Float FilteredColor = VectorLoad((const float*)&Sample);
+
+				// Apply computed alpha scales to each channel
+				FilteredColor = VectorMultiply(FilteredColor, VectorLoad((const float*)&AlphaScale));
+
+				// Set the destination pixel.
+				VectorStore(FilteredColor, (float*)&DestImageData.Access(DestX, DestY));
+			}
+		}
+	});
+}
+
+template <EMipGenAddressMode AddressMode>
+static void GenerateMip2x2(const FImageView2D& SourceImageData, FImageView2D& DestImageData, const FVector4f AlphaScale, uint32 ScaleFactor)
+{
+	int32 NumRowsEachJob;
+	int32 NumJobs = ImageParallelForComputeNumJobsForRows(NumRowsEachJob, DestImageData.SizeX, DestImageData.SizeY);
+
+	ParallelFor(TEXT("Texture.GenerateMip2x2.PF"), NumJobs, 1, [&](int32 Index)
+	{
+		int32 StartIndex = Index * NumRowsEachJob;
+		int32 EndIndex = FMath::Min(StartIndex + NumRowsEachJob, DestImageData.SizeY);
+		for (int32 DestY = StartIndex; DestY < EndIndex; ++DestY)
+		{
+			for (int32 DestX = 0; DestX < DestImageData.SizeX; DestX++)
+			{
+				const int32 SourceX = DestX * ScaleFactor;
+				const int32 SourceY = DestY * ScaleFactor;
+
+				VectorRegister4Float A = VectorLoad((const float*)&LookupSourceMip<AddressMode>(SourceImageData, SourceX + 0, SourceY + 0));
+				VectorRegister4Float B = VectorLoad((const float*)&LookupSourceMip<AddressMode>(SourceImageData, SourceX + 1, SourceY + 0));
+				VectorRegister4Float C = VectorLoad((const float*)&LookupSourceMip<AddressMode>(SourceImageData, SourceX + 0, SourceY + 1));
+				VectorRegister4Float D = VectorLoad((const float*)&LookupSourceMip<AddressMode>(SourceImageData, SourceX + 1, SourceY + 1));
+				VectorRegister4Float FilteredColor = VectorAdd(VectorAdd(VectorAdd(A, B), C), D);
+				FilteredColor = VectorMultiply(FilteredColor, VectorSetFloat1(0.25f));
+
+				// Apply computed alpha scales to each channel
+				FilteredColor = VectorMultiply(FilteredColor, VectorLoad((const float*)&AlphaScale));
+
+				// Set the destination pixel.
+				VectorStore(FilteredColor, (float*)&DestImageData.Access(DestX, DestY));
+			}
+		}
+	});
+}
+
+template <EMipGenAddressMode AddressMode, bool bSharpenWithoutColorShift, int KernelSize = 0>
+static void GenerateMipSharpened(
+	const FImageView2D& SourceImageData,
+	FImageView2D& DestImageData,
+	const FImageKernel2D& Kernel,
+	const FVector4f AlphaScale,
+	uint32 ScaleFactor)
+{
+	int32 NumRowsEachJob;
+	int32 NumJobs = ImageParallelForComputeNumJobsForRows(NumRowsEachJob, DestImageData.SizeX, DestImageData.SizeY);
+
+	ParallelFor(TEXT("Texture.GenerateMipSharpened.PF"), NumJobs, 1, [&](int32 Index)
+	{
+		// In case kernel size is passed as template argument use it as constant
+		// This will allow compiler to unroll inner loops below
+		const int32 KernelFilterTableSize = KernelSize ? KernelSize : (int32)Kernel.GetFilterTableSize();
+
+		// if KernelFilterTableSize is odd, centered in-place filter can be applied
+		// KernelFilterTableSize should be even for standard down-sampling
+		const int32 KernelCenter = (KernelFilterTableSize - 1) / 2;
+
+		int32 StartIndex = Index * NumRowsEachJob;
+		int32 EndIndex = FMath::Min(StartIndex + NumRowsEachJob, DestImageData.SizeY);
+		for (int32 DestY = StartIndex; DestY < EndIndex; ++DestY)
+		{
+			for (int32 DestX = 0; DestX < DestImageData.SizeX; DestX++)
+			{
+				const int32 SourceX = DestX * ScaleFactor;
+				const int32 SourceY = DestY * ScaleFactor;
+
+				VectorRegister4Float Color = VectorZeroFloat();
+				for (int32 KernelY = 0; KernelY < KernelFilterTableSize; ++KernelY)
+				{
+					for (int32 KernelX = 0; KernelX < KernelFilterTableSize; ++KernelX)
+					{
+						const FLinearColor& Sample = LookupSourceMip<AddressMode>(SourceImageData, SourceX + KernelX - KernelCenter, SourceY + KernelY - KernelCenter);
+						VectorRegister4Float Weight = VectorSetFloat1(Kernel.GetAt(KernelX, KernelY));
+						VectorRegister4Float WeightSample = VectorMultiply(Weight, VectorLoad((const float*)&Sample));
+						Color = VectorAdd(Color, WeightSample);
+					}
+				}
+
+				// This condition will be optimized away because it is constant template argument
+				if (bSharpenWithoutColorShift)
+				{
+					VectorRegister4Float SharpenedColor = Color;
+
+					// Luminace weights from FLinearColor::GetLuminance() function
+					VectorRegister4Float LuminanceWeights = MakeVectorRegisterFloat(0.3f, 0.59f, 0.11f, 0.f);
+					VectorRegister4Float NewLuminance = VectorDot3(SharpenedColor, LuminanceWeights);
+
+					// simple 2x2 kernel to compute the color
+					VectorRegister4Float A = VectorLoad((const float*)&LookupSourceMip<AddressMode>(SourceImageData, SourceX + 0, SourceY + 0));
+					VectorRegister4Float B = VectorLoad((const float*)&LookupSourceMip<AddressMode>(SourceImageData, SourceX + 1, SourceY + 0));
+					VectorRegister4Float C = VectorLoad((const float*)&LookupSourceMip<AddressMode>(SourceImageData, SourceX + 0, SourceY + 1));
+					VectorRegister4Float D = VectorLoad((const float*)&LookupSourceMip<AddressMode>(SourceImageData, SourceX + 1, SourceY + 1));
+					VectorRegister4Float FilteredColor = VectorAdd(VectorAdd(VectorAdd(A, B), C), D);
+					FilteredColor = VectorMultiply(FilteredColor, VectorSetFloat1(0.25f));
+
+					VectorRegister4Float OldLuminance = VectorDot3(FilteredColor, LuminanceWeights);
+
+					// if (OldLuminance > 0.001f) FilteredColor.RGB *= NewLuminance / OldLuminance;
+					VectorRegister4Float CompareMask = VectorCompareGT(OldLuminance, VectorSetFloat1(0.001f));
+					VectorRegister4Float Temp = VectorMultiply(FilteredColor, VectorDivide(NewLuminance, OldLuminance));
+					FilteredColor = VectorSelect(CompareMask, Temp, FilteredColor);
+
+					// FilteredColor.A = SharpenedColor.A
+					VectorRegister4Float AlphaMask = MakeVectorRegisterFloatMask(0, 0, 0, 0xffffffff);
+					FilteredColor = VectorSelect(AlphaMask, SharpenedColor, FilteredColor);
+
+					// Apply computed alpha scales to each channel
+					FilteredColor = VectorMultiply(FilteredColor, VectorLoad((const float*)&AlphaScale));
+
+					// Set the destination pixel.
+					VectorStore(FilteredColor, (float*)&DestImageData.Access(DestX, DestY));
+				}
+				else
+				{
+					// Apply computed alpha scales to each channel
+					Color = VectorMultiply(Color, VectorLoad((const float*)&AlphaScale));
+
+					// Set the destination pixel.
+					VectorStore(Color, (float*)&DestImageData.Access(DestX, DestY));
+				}
+
 			}
 		}
 	});
@@ -614,110 +775,46 @@ static void GenerateSharpenedMipB8G8R8A8Templ(
 		return;
 	}
 
-	// if KernelFilterTableSize is odd, centered in-place filter can be applied
-	// KernelFilterTableSize should be even for standard down-sampling
-	const int32 KernelCenter = (KernelFilterTableSize - 1) / 2;
-
 	FVector4f AlphaScale(1, 1, 1, 1);
 	if (bDoScaleMipsForAlphaCoverage)
 	{
 		AlphaScale = ComputeAlphaScale(AlphaCoverages, AlphaThresholds, SourceImageData);
 	}
-	
-	int32 NumRowsEachJob;
-	int32 NumJobs = ImageParallelForComputeNumJobsForRows(NumRowsEachJob,DestImageData.SizeX,DestImageData.SizeY);
 
-	ParallelFor( TEXT("Texture.GenerateSharpenedMip.PF"),NumJobs,1, [&](int32 Index)
+	if (bUnfiltered)
 	{
-		int32 StartIndex = Index * NumRowsEachJob;
-		int32 EndIndex = FMath::Min(StartIndex + NumRowsEachJob, DestImageData.SizeY);
-		for (int32 DestY = StartIndex; DestY < EndIndex; ++DestY)
-		{
-			for ( int32 DestX = 0;DestX < DestImageData.SizeX; DestX++ )
-			{
-				const int32 SourceX = DestX * ScaleFactor;
-				const int32 SourceY = DestY * ScaleFactor;
-
-				FLinearColor FilteredColor(0, 0, 0, 0);
-
-				if ( bUnfiltered )
-				{
-					FilteredColor = LookupSourceMip<AddressMode>(SourceImageData, SourceX + 0, SourceY + 0);
-				}
-				else if ( KernelFilterTableSize == 2 )
-				{
-					// simple 2x2 kernel to compute the color
-					// check for this case early so that bSharpenWithoutColorShift isn't run on 2x2
-
-					FilteredColor =
-						( LookupSourceMip<AddressMode>( SourceImageData, SourceX + 0, SourceY + 0 )
-						+ LookupSourceMip<AddressMode>( SourceImageData, SourceX + 1, SourceY + 0 )
-						+ LookupSourceMip<AddressMode>( SourceImageData, SourceX + 0, SourceY + 1 )
-						+ LookupSourceMip<AddressMode>( SourceImageData, SourceX + 1, SourceY + 1 ) ) * 0.25f;
-				}
-				else if ( bSharpenWithoutColorShift )
-				{
-					// don't use this if KernelFilterTableSize == 2 (box)
-					//	because it just does the same thing twice
-
-					FLinearColor SharpenedColor(0, 0, 0, 0);
-
-					for ( int32 KernelY = 0; KernelY < KernelFilterTableSize;  ++KernelY )
-					{
-						for ( int32 KernelX = 0; KernelX < KernelFilterTableSize;  ++KernelX )
-						{
-							float Weight = Kernel.GetAt( KernelX, KernelY );
-							FLinearColor Sample = LookupSourceMip<AddressMode>( SourceImageData, SourceX + KernelX - KernelCenter, SourceY + KernelY - KernelCenter );
-							SharpenedColor += Weight * Sample;
-						}
-					}
-
-					float NewLuminance = SharpenedColor.GetLuminance();
-
-					// simple 2x2 kernel to compute the color
-					FilteredColor =
-						( LookupSourceMip<AddressMode>( SourceImageData, SourceX + 0, SourceY + 0 )
-						+ LookupSourceMip<AddressMode>( SourceImageData, SourceX + 1, SourceY + 0 )
-						+ LookupSourceMip<AddressMode>( SourceImageData, SourceX + 0, SourceY + 1 )
-						+ LookupSourceMip<AddressMode>( SourceImageData, SourceX + 1, SourceY + 1 ) ) * 0.25f;
-
-					float OldLuminance = FilteredColor.GetLuminance();
-
-					if ( OldLuminance > 0.001f )
-					{
-						float Factor = NewLuminance / OldLuminance;
-						FilteredColor.R *= Factor;
-						FilteredColor.G *= Factor;
-						FilteredColor.B *= Factor;
-					}
-
-					// We also want to sharpen the alpha channel (was missing before)
-					FilteredColor.A = SharpenedColor.A;
-				}
-				else
-				{
-					for ( int32 KernelY = 0; KernelY < KernelFilterTableSize;  ++KernelY )
-					{
-						for ( int32 KernelX = 0; KernelX < KernelFilterTableSize;  ++KernelX )
-						{
-							float Weight = Kernel.GetAt( KernelX, KernelY );
-							FLinearColor Sample = LookupSourceMip<AddressMode>( SourceImageData, SourceX + KernelX - KernelCenter, SourceY + KernelY - KernelCenter );
-							FilteredColor += Weight	* Sample;
-						}
-					}
-				}
+		GenerateMipUnfiltered<AddressMode>(SourceImageData, DestImageData, AlphaScale, ScaleFactor);
+		return;
+	}
 
-				// Apply computed alpha scales to each channel		
-				FilteredColor.R *= AlphaScale.X;
-				FilteredColor.G *= AlphaScale.Y;
-				FilteredColor.B *= AlphaScale.Z;
-				FilteredColor.A *= AlphaScale.W;
+	if (KernelFilterTableSize == 2)
+	{
+		GenerateMip2x2<AddressMode>(SourceImageData, DestImageData, AlphaScale, ScaleFactor);
+		return;
+	}
 
-				// Set the destination pixel.
-				DestImageData.Access(DestX, DestY) = FilteredColor;
-			}
+	if (bSharpenWithoutColorShift)
+	{
+		if (KernelFilterTableSize == 8)
+		{
+			GenerateMipSharpened<AddressMode, true, 8>(SourceImageData, DestImageData, Kernel, AlphaScale, ScaleFactor);
 		}
-	});
+		else
+		{
+			GenerateMipSharpened<AddressMode, true>(SourceImageData, DestImageData, Kernel, AlphaScale, ScaleFactor);
+		}
+	}
+	else
+	{
+		if (KernelFilterTableSize == 8)
+		{
+			GenerateMipSharpened<AddressMode, false, 8>(SourceImageData, DestImageData, Kernel, AlphaScale, ScaleFactor);
+		}
+		else
+		{
+			GenerateMipSharpened<AddressMode, false>(SourceImageData, DestImageData, Kernel, AlphaScale, ScaleFactor);
+		}
+	}
 }
 
 // to switch conveniently between different texture wrapping modes for the mip map generation
