TextureCompressorModule simplify and improve parallel mip encoding

Pixel count in mip maps is in a geometric progression: each 2D mip has 4x as many pixels as the one below it,
and 3x as many as the sum of all mips below it. Therefore, a parallel for over mip levels is not very useful:
the vast majority of work is in the base mip, and leaving the base mip (which is on our critical path) for a
worker to pick up adds whatever time it takes for a worker to get started on it to our expected total time.

Therefore, change the logic so the calling thread grabs the base mip itself, and spawns exactly one other
task before to process all the other mips. This should reduce the number of tiny, useless tasks we spawn,
and because this single async task is expected to take 1/3 of the time as the base mip task, we have a decent
amount of slack for queuing delays without it showing up as a real wait time.

diff --git a/Engine/Source/Developer/TextureCompressor/Private/TextureCompressorModule.cpp b/Engine/Source/Developer/TextureCompressor/Private/TextureCompressorModule.cpp
--- a/Engine/Source/Developer/TextureCompressor/Private/TextureCompressorModule.cpp
+++ b/Engine/Source/Developer/TextureCompressor/Private/TextureCompressorModule.cpp
@@ -11,6 +11,7 @@
 #include "TextureFormatManager.h"
 #include "Interfaces/ITextureFormat.h"
 #include "Misc/Paths.h"
+#include "Tasks/Task.h"
 #include "ImageCore.h"
 #include <cmath>
 
@@ -2445,76 +2446,6 @@ static void ApplyCompositeTexture(FImage& RoughnessSourceMips, const FImage& Nor
 	Image Compression.
 ------------------------------------------------------------------------------*/
 
-/**
- * Asynchronous compression, used for compressing mips simultaneously.
- */
-class FAsyncCompressionWorker
-{
-public:
-	/**
-	 * Initializes the data and creates the async compression task.
-	 */
-	FAsyncCompressionWorker(const ITextureFormat* InTextureFormat, const FImage* InImages, uint32 InNumImages, const FTextureBuildSettings& InBuildSettings, FStringView InDebugTexturePathName, bool bInImageHasAlphaChannel, uint32 InExtData)
-		: TextureFormat(*InTextureFormat)
-		, SourceImages(InImages)
-		, BuildSettings(InBuildSettings)
-		, bImageHasAlphaChannel(bInImageHasAlphaChannel)
-		, ExtData(InExtData)
-		, NumImages(InNumImages)
-		, bCompressionResults(false)
-		, DebugTexturePathName(InDebugTexturePathName)
-	{
-	}
-
-	/**
-	 * Compresses the texture
-	 */
-	void DoWork()
-	{
-		TRACE_CPUPROFILER_EVENT_SCOPE(Texture.CompressImage);
-
-		bCompressionResults = TextureFormat.CompressImageEx(
-			SourceImages,
-			NumImages,
-			BuildSettings,
-			DebugTexturePathName,
-			bImageHasAlphaChannel,
-			ExtData,
-			CompressedImage
-			);
-	}
-
-	/**
-	 * Transfer the result of the compression to the OutCompressedImage
-	 * Can only be called once
-	 */
-	bool ConsumeCompressionResults(FCompressedImage2D& OutCompressedImage)
-	{
-		OutCompressedImage = MoveTemp(CompressedImage);
-		return bCompressionResults;
-	}
-
-private:
-
-	/** Texture format interface with which to compress. */
-	const ITextureFormat& TextureFormat;
-	/** The image(s) to compress. */
-	const FImage* SourceImages;
-	/** The resulting compressed image. */
-	FCompressedImage2D CompressedImage;
-	/** Build settings. */
-	FTextureBuildSettings BuildSettings;
-	/** true if the image has a non-white alpha channel. */
-	bool bImageHasAlphaChannel;
-	/** Extra data that the format may want to pass to each Compress call */
-	uint32 ExtData;
-	/** For miptails with multiple images going in to one, this is the number of them */
-	uint32 NumImages;
-	/** true if compression was successful. */
-	bool bCompressionResults;
-	FStringView DebugTexturePathName;
-};
-
 // compress mip-maps in InMipChain and add mips to Texture, might alter the source content
 static bool CompressMipChain(
 	const ITextureFormat* TextureFormat,
@@ -2540,82 +2471,74 @@ static bool CompressMipChain(
 	const int32 MinAsyncCompressionSize = 512;
 	const bool bAllowParallelBuild = TextureFormat->AllowParallelBuild();
 	bool bCompressionSucceeded = true;
-	int32 FirstMipTailIndex = MipCount;
-	uint32 StartCycles = FPlatformTime::Cycles();
 
-	// check if we need to merge mips together into tail
+	// Mip tail is when the last few mips get grouped together in the hardware layout.
+	// Treat not having a mip tail as having a mip tail with 1 mip in it, which is
+	// equivalent and lets us simplify the logic.
+	int32 FirstMipTailIndex = MipCount - 1;
+	int32 MipTailCount = 1;
+
 	if (CompressorCaps.NumMipsInTail > 1)
 	{
-		FirstMipTailIndex = MipCount - CompressorCaps.NumMipsInTail;
+		MipTailCount = CompressorCaps.NumMipsInTail;
+		FirstMipTailIndex = MipCount - MipTailCount;
 	}
 
-	OutMips.Empty(MipCount);
-	TArray<FAsyncCompressionWorker> AsyncCompressionTasks;
-	AsyncCompressionTasks.Reserve(MipCount);
+	uint32 StartCycles = FPlatformTime::Cycles();
 
-	struct PreWork
-	{
-		int32 MipIndex;
-		const FImage& SrcMip;
-		FCompressedImage2D& DestMip;
-	};
-	TArray<PreWork> PreWorkTasks;
-	PreWorkTasks.Reserve(MipCount);
+	// Set up one task for the base mip, one task for everything after. Since each mip level
+	// has 4x the pixels as the one below it (8x for volumes), work for mip levels is highly
+	// unbalanced and there's not much use spawning extra tasks past that: for a 2D texture,
+	// the entire tail after the base mip (all remaining mips combined) has 1/3 the number of
+	// pixels the base mip does.
+	OutMips.Empty(MipCount);
+	OutMips.AddDefaulted(MipCount);
 
-	for (int32 MipIndex = 0; MipIndex < MipCount; ++MipIndex)
+	auto ProcessMips =
+		[&TextureFormat, &MipChain, &OutMips, FirstMipTailIndex, MipTailCount, &CompressorCaps, &Settings, &DebugTexturePathName, bImageHasAlphaChannel](int32 MipBegin, int32 MipEnd)
 	{
-		const FImage& SrcMip = MipChain[MipIndex];
-		FCompressedImage2D& DestMip = *new(OutMips) FCompressedImage2D;
+		bool bSuccess = true;
 
-		if (MipIndex > FirstMipTailIndex)
+		for (int32 MipIndex = MipBegin; MipIndex < MipEnd; ++MipIndex)
 		{
-			continue;
-		}
-		else if (bAllowParallelBuild && FMath::Min(SrcMip.SizeX, SrcMip.SizeY) >= MinAsyncCompressionSize)
-		{
-			AsyncCompressionTasks.Emplace(
-				TextureFormat,
-				&SrcMip,
-				MipIndex == FirstMipTailIndex ? CompressorCaps.NumMipsInTail : 1, // number of mips pointed to by SrcMip
+			TRACE_CPUPROFILER_EVENT_SCOPE(Texture.CompressImage);
+
+			bSuccess = bSuccess && TextureFormat->CompressImageEx(
+				&MipChain[MipIndex],
+				MipIndex == FirstMipTailIndex ? MipTailCount : 1, // number of mips pointed to by SrcMip
 				Settings,
 				DebugTexturePathName,
 				bImageHasAlphaChannel,
-				CompressorCaps.ExtData
+				CompressorCaps.ExtData,
+				OutMips[MipIndex]
 			);
 		}
-		else
-		{
-			PreWorkTasks.Emplace(PreWork { MipIndex, SrcMip, DestMip });
-		}
-	}
 
-	ParallelForWithPreWork( TEXT("Texture.CompressMipChain.PF"),AsyncCompressionTasks.Num(),1,
-		[&AsyncCompressionTasks](int32 TaskIndex)
-		{
-			AsyncCompressionTasks[TaskIndex].DoWork();
-		},
-		[&PreWorkTasks, &TextureFormat, &OutMips, &bCompressionSucceeded, &CompressorCaps, &Settings, DebugTexturePathName, FirstMipTailIndex, bImageHasAlphaChannel]()
-		{
-			for (PreWork& Work : PreWorkTasks)
+		return bSuccess;
+	};
+
+	if (bAllowParallelBuild &&
+		FirstMipTailIndex > 0 &&
+		FMath::Min(MipChain[0].SizeX, MipChain[0].SizeY) >= MinAsyncCompressionSize)
+	{
+		// Spawn async job to compress all mips below base
+		auto AsyncTask = UE::Tasks::Launch(TEXT("Texture.CompressLowerMips"),
+			[&ProcessMips, FirstMipTailIndex]()
 			{
-				bCompressionSucceeded = bCompressionSucceeded && TextureFormat->CompressImageEx(
-					&Work.SrcMip,
-					Work.MipIndex == FirstMipTailIndex ? CompressorCaps.NumMipsInTail : 1, // number of mips pointed to by SrcMip
-					Settings,
-					DebugTexturePathName,
-					bImageHasAlphaChannel,
-					CompressorCaps.ExtData,
-					Work.DestMip
-				);
-			}
-		}, 
-		EParallelForFlags::Unbalanced);
+				TRACE_CPUPROFILER_EVENT_SCOPE(Texture.CompressLowerMips);
+				return ProcessMips(1, FirstMipTailIndex + 1);
+			},
+			LowLevelTasks::ETaskPriority::BackgroundNormal
+		);
 
-	for (int32 TaskIndex = 0; TaskIndex < AsyncCompressionTasks.Num(); ++TaskIndex)
+		// Compress base mip on this thread, join with async compress of other mips
+		bCompressionSucceeded = ProcessMips(0, 1);
+		bCompressionSucceeded &= AsyncTask.GetResult();
+	}
+	else
 	{
-		FAsyncCompressionWorker& AsynTask = AsyncCompressionTasks[TaskIndex];
-		FCompressedImage2D& DestMip = OutMips[TaskIndex];
-		bCompressionSucceeded = bCompressionSucceeded && AsynTask.ConsumeCompressionResults(DestMip);
+		// Compress all mips at once on this thread
+		bCompressionSucceeded = ProcessMips(0, FirstMipTailIndex + 1);
 	}
 
 	for (int32 MipIndex = FirstMipTailIndex + 1; MipIndex < MipCount; ++MipIndex)
