fix share Texture validation on Import code for all texture types
turns off mip maps for nonpow2 cubemaps
fixes check importing nonpow2 cubemaps

diff --git a/Engine/Plugins/Media/BinkMedia/Source/BinkMediaPlayer/Public/BinkMediaTexture.h b/Engine/Plugins/Media/BinkMedia/Source/BinkMediaPlayer/Public/BinkMediaTexture.h
--- a/Engine/Plugins/Media/BinkMedia/Source/BinkMediaPlayer/Public/BinkMediaTexture.h
+++ b/Engine/Plugins/Media/BinkMedia/Source/BinkMediaPlayer/Public/BinkMediaTexture.h
@@ -14,6 +14,9 @@ enum TextureAddress;
 
 /**
  * Implements a texture asset for rendering video tracks from UBinkMediaPlayer assets.
+ * 
+ * note: derives directly from UTexture, not from UTexture2D or UTexture2DDynamic
+ *    maybe should have been UTexture2DDynamic?
  */
 UCLASS(hidecategories=(Compression, LevelOfDetail, Object))
 class BINKMEDIAPLAYER_API UBinkMediaTexture : public UTexture 
@@ -38,6 +41,7 @@ class BINKMEDIAPLAYER_API UBinkMediaTexture : public UTexture
 	virtual float GetSurfaceHeight() const override { return CachedDimensions.Y; }
 	virtual float GetSurfaceDepth() const override { return 0; }
 	virtual uint32 GetSurfaceArraySize() const override { return 0; }
+	virtual ETextureClass GetTextureClass() const { return ETextureClass::Other2DNoSource; }
 
 	// UObject overrides.
 
diff --git a/Engine/Source/Developer/TextureCompressor/Private/TextureCompressorModule.cpp b/Engine/Source/Developer/TextureCompressor/Private/TextureCompressorModule.cpp
--- a/Engine/Source/Developer/TextureCompressor/Private/TextureCompressorModule.cpp
+++ b/Engine/Source/Developer/TextureCompressor/Private/TextureCompressorModule.cpp
@@ -1183,6 +1183,13 @@ void ITextureCompressorModule::GenerateMipChain(
 {
 	TRACE_CPUPROFILER_EVENT_SCOPE(Texture.GenerateMipChain);
 
+	// MipChainDepth is the number more to make, OutMipChain has some already
+	// typically BaseImage == OutMipChain.Last()
+	if ( MipChainDepth == 0 )
+	{
+		return;
+	}
+
 	check(BaseImage.Format == ERawImageFormat::RGBA32F);
 
 	const FImage& BaseMip = BaseImage;
@@ -1252,6 +1259,13 @@ void ITextureCompressorModule::GenerateMipChain(
 		check(IntermediateSrcPtr && IntermediateDstPtr);
 		const FImage& IntermediateSrc = *IntermediateSrcPtr;
 		FImage& IntermediateDst = *IntermediateDstPtr;
+		
+		if ( IntermediateSrc.SizeX == 1 && IntermediateSrc.SizeY == 1 && (!Settings.bVolume || IntermediateSrc.NumSlices == 1))
+		{
+			// should not have been called, starting mip is already small enough
+			check(0);
+			break;
+		}
 
 		// add new mip to TArray<FImage> &OutMipChain :
 		//	placement new on TArray does AddUninitialized then constructs in the last element
@@ -2806,6 +2820,13 @@ private:
 			}
 		}
 
+		// @@!! this is broken for non-pow2 sizes
+		//	if size is non-pow2, outer code should have set TMGS_NoMipmaps
+		
+		// problem 1 : this should be FloorLogTwo not Ceil
+		// problem 2 : if you then do that, it's broken for padding to pow2 mode, because the pad is after the mip count calculation
+		// problem 3 : pad+leave existing is broken
+
 		// Determine the maximum possible mip counts for source and dest.
 		const int32 MaxSourceMipCount = bLongLatCubemap ?
 			1 + FMath::CeilLogTwo(ComputeLongLatCubemapExtents(InSourceMips[0], BuildSettings.MaxTextureResolution)) :
@@ -2823,16 +2844,27 @@ private:
 
 		if (BuildSettings.MipGenSettings == TMGS_LeaveExistingMips)
 		{
-			NumOutputMips = InSourceMips.Num() - StartMip;
-			if (NumOutputMips <= 0)
+			// @@!! if LeaveExistingMips and we want 10 mips but are only given 8, we don't make the lower ones?
+			//	 is that right?
+			int32 NumExistingMips = NumSourceMips - StartMip;
+			check( NumExistingMips <= NumOutputMips );
+
+			if (NumExistingMips <= 0)
 			{
 				// We can't generate 0 mip maps
 				UE_LOG(LogTextureCompressor, Warning,
-					TEXT("The source image has %d mips while the first mip would be %d. Please verify the maximun texture size or change the mips gen settings."),
+					TEXT("The source image has %d mips while the first mip would be %d. Please verify the maximum texture size or change the mips gen settings."),
 					NumSourceMips,
 					StartMip);
 				return false;
 			}
+			else
+			{
+				// if NumExistingMips < NumOutputMips
+				// a non-complete mip tail will be generated
+
+				NumOutputMips = NumExistingMips;
+			}
 		}
 
 		NumOutputMips = FMath::Min(NumOutputMips, MaxDestMipCount);
@@ -2846,6 +2878,8 @@ private:
 		TArray<FImage> PaddedSourceMips;
 
 		{
+			// @@!! Padded + leave existing is broken but that's not enforced ?
+
 			const FImage& FirstSourceMipImage = InSourceMips[0];
 			int32 TargetTextureSizeX = FirstSourceMipImage.SizeX;
 			int32 TargetTextureSizeY = FirstSourceMipImage.SizeY;
@@ -2870,6 +2904,7 @@ private:
 			case ETexturePowerOfTwoSetting::PadToSquarePowerOfTwo:
 				bPadOrStretchTexture = true;
 				TargetTextureSizeX = TargetTextureSizeY = FMath::Max3<int32>(PowerOfTwoTextureSizeX, PowerOfTwoTextureSizeY, PowerOfTwoTextureSizeZ);
+				// @@!! doesn't change or check TargetTextureSizeZ ?
 				break;
 
 			default:
@@ -2941,6 +2976,8 @@ private:
 		TArray<FImage> GeneratedSourceMips;
 		if (bBuildSourceImage)
 		{
+			// should not get in here with LeaveExistingMips because NumExistingMips would be < 0
+
 			// the source is larger than the compressor allows and no mip image exists to act as a smaller source.
 			// We must generate a suitable source image:
 			bool bSuitableFormat = PostOptionalUpscaleSourceMips.Last().Format == ERawImageFormat::RGBA32F;
@@ -2975,6 +3012,7 @@ private:
 
 			check(GeneratedSourceMips.Num() != 0);
 			// Note: The newly generated mip chain does not include the original top level mip.
+			check( StartMip > 0 );
 			StartMip--;
 		}
 
@@ -2985,10 +3023,12 @@ private:
 		check(StartMip < SourceMips.Num());
 		int32 CopyCount = SourceMips.Num() - StartMip;
 
+		int32 GenerateCount = NumOutputMips - CopyCount;
+
 		// avoid converting to RGBA32F linear format if there's no need for any extra processing of pixels
 		// image will be left in BGRA8 format if possible
 		const bool bNeedAdjustImageColors = NeedAdjustImageColors(BuildSettings);
-		const bool bLinearize = bNeedLinearize || NumOutputMips != 1 || BuildSettings.bRenormalizeTopMip || (BuildSettings.Downscale > 1.f)
+		const bool bLinearize = bNeedLinearize || (GenerateCount > 0) || BuildSettings.bRenormalizeTopMip || (BuildSettings.Downscale > 1.f)
 			|| BuildSettings.bHasColorSpaceDefinition || BuildSettings.bComputeBokehAlpha || BuildSettings.bFlipGreenChannel
 			|| BuildSettings.bReplicateRed || BuildSettings.bReplicateAlpha || BuildSettings.bApplyYCoCgBlockScale
 			|| BuildSettings.SourceEncodingOverride != 0 || bNeedAdjustImageColors;
@@ -3075,6 +3115,9 @@ private:
 			}
 		}
 
+		check( OutMipChain.Num() == CopyCount );
+		check( GenerateCount == NumOutputMips - OutMipChain.Num() );
+
 		// Generate any missing mips in the chain.
 		if (NumOutputMips > OutMipChain.Num())
 		{
@@ -3085,7 +3128,10 @@ private:
 			}
 			else
 			{
-				GenerateMipChain(BuildSettings, OutMipChain.Last(), OutMipChain);
+				// @@!! GenerateMipChain should bring us up to NumOutputMips
+				//	but it doesn't take NumOutputMips as a param, makes its own decision
+
+				GenerateMipChain(BuildSettings, OutMipChain.Last(), OutMipChain, MAX_uint32);
 			}
 		}
 		check(OutMipChain.Num() == NumOutputMips);
diff --git a/Engine/Source/Developer/TextureCompressor/Public/TextureCompressorModule.h b/Engine/Source/Developer/TextureCompressor/Public/TextureCompressorModule.h
--- a/Engine/Source/Developer/TextureCompressor/Public/TextureCompressorModule.h
+++ b/Engine/Source/Developer/TextureCompressor/Public/TextureCompressorModule.h
@@ -318,7 +318,7 @@ public:
 		const FTextureBuildSettings& Settings,
 		const FImage& BaseImage,
 		TArray<FImage> &OutMipChain,
-		uint32 MipChainDepth = MAX_uint32
+		uint32 MipChainDepth
 		);
 
 	/**
diff --git a/Engine/Source/Editor/UnrealEd/Private/Factories/EditorFactories.cpp b/Engine/Source/Editor/UnrealEd/Private/Factories/EditorFactories.cpp
--- a/Engine/Source/Editor/UnrealEd/Private/Factories/EditorFactories.cpp
+++ b/Engine/Source/Editor/UnrealEd/Private/Factories/EditorFactories.cpp
@@ -4219,6 +4219,8 @@ UObject* UTextureFactory::FactoryCreateBinary
 	CompressionSettings = Texture->CompressionSettings;
 
 	// Figure out whether we're using a normal map LOD group.
+	// @@!! LODGroup is a member var on TextureFactory not pulled from Texture->LODGroup ?
+
 	bool bIsNormalMapLODGroup = false;
 	if( LODGroup == TEXTUREGROUP_WorldNormalMap 
 	||	LODGroup == TEXTUREGROUP_CharacterNormalMap
@@ -4366,6 +4368,8 @@ UObject* UTextureFactory::FactoryCreateBinary
 
 	GEditor->GetEditorSubsystem<UImportSubsystem>()->BroadcastAssetPostImport(this, Texture);
 
+	// PostEditChange calls ValidateSettingsAfterImportOrEdit();
+
 	// Invalidate any materials using the newly imported texture. (occurs if you import over an existing texture)
 	Texture->PostEditChange();
 
diff --git a/Engine/Source/Runtime/AugmentedReality/Public/ARTextures.h b/Engine/Source/Runtime/AugmentedReality/Public/ARTextures.h
--- a/Engine/Source/Runtime/AugmentedReality/Public/ARTextures.h
+++ b/Engine/Source/Runtime/AugmentedReality/Public/ARTextures.h
@@ -24,6 +24,7 @@ enum class EARTextureType : uint8
 /**
  * Base class for all AR texture types.
  * Derived from UTexture instead of UTexture2D because UTexture2D is all about streaming and source art
+ * ? probably should have been UTexture2DDynamic
  */
 UCLASS(Abstract, BlueprintType)
 class AUGMENTEDREALITY_API UARTexture : public UTexture
@@ -49,6 +50,7 @@ public:
 	virtual float GetSurfaceDepth() const override { return 0; }
 	virtual uint32 GetSurfaceArraySize() const override { return 0; }
 	virtual FGuid GetExternalTextureGuid() const override { return ExternalTextureGuid; }
+	virtual ETextureClass GetTextureClass() const { return ETextureClass::Other2DNoSource; }
 	// End UTexture interface
 	
 	/** The type of texture this is */
diff --git a/Engine/Source/Runtime/ColorManagement/Private/ColorSpace.cpp b/Engine/Source/Runtime/ColorManagement/Private/ColorSpace.cpp
--- a/Engine/Source/Runtime/ColorManagement/Private/ColorSpace.cpp
+++ b/Engine/Source/Runtime/ColorManagement/Private/ColorSpace.cpp
@@ -48,6 +48,7 @@ FColorSpace::FColorSpace(EColorSpace ColorSpaceType)
 	switch (ColorSpaceType)
 	{
 	case EColorSpace::None:
+		// leaves FVector2d::Zero from the constructor
 		break;
 	case EColorSpace::sRGB:
 		Chromaticities[0] = FVector2d(0.64, 0.33);
diff --git a/Engine/Source/Runtime/Engine/Classes/Engine/Texture.h b/Engine/Source/Runtime/Engine/Classes/Engine/Texture.h
--- a/Engine/Source/Runtime/Engine/Classes/Engine/Texture.h
+++ b/Engine/Source/Runtime/Engine/Classes/Engine/Texture.h
@@ -314,7 +314,7 @@ struct FTextureSource
 	/** Computes the number of bytes per-pixel. */
 	ENGINE_API int32 GetBytesPerPixel(int32 LayerIndex = 0) const;
 
-	/** Return true if the source data is power-of-2. */
+	/** Return true if the source XY size is power-of-2.  Does not check Z size for volumes.  */
 	ENGINE_API bool IsPowerOfTwo(int32 BlockIndex = 0) const;
 
 	/** Returns true if source art is available. */
@@ -1061,6 +1061,7 @@ class UTexture : public UStreamableRenderAsset, public IInterface_AssetUserData,
 	--------------------------------------------------------------------------*/
 
 #if WITH_EDITORONLY_DATA
+	/* Dynamic textures will have ! Source.IsValid() */
 	UPROPERTY()
 	FTextureSource Source;
 #endif
@@ -1367,6 +1368,9 @@ public:
 	ENGINE_API virtual void ImportCustomProperties(const TCHAR* SourceText, FFeedbackContext* Warn) override;
 	ENGINE_API virtual void PostEditImport() override;
 
+	/** Get Texture Class */
+	ENGINE_API virtual ETextureClass GetTextureClass() const PURE_VIRTUAL(UTexture::GetTextureClass, return ETextureClass::Invalid; );
+
 	/**
 	 * Resets the resource for the texture.
 	 */
@@ -1516,13 +1520,26 @@ public:
 	ENGINE_API void GetDefaultFormatSettings(FTextureFormatSettings& OutSettings) const;
 	
 	ENGINE_API ETextureEncodeSpeed GetDesiredEncodeSpeed() const;
+		
+	/** Ensure settings are valid after import or edit; this is called by PostEditChange. */
+	ENGINE_API virtual void ValidateSettingsAfterImportOrEdit();
+#endif // WITH_EDITOR
 
 	ENGINE_API EGammaSpace GetGammaSpace() const
 	{
 		// note: does not validate that the Format respects gamma (TextureSource::GetGammaSpace does)
-		return SRGB ? ( bUseLegacyGamma ? EGammaSpace::Pow22 : EGammaSpace::sRGB ) : EGammaSpace::Linear;
+		
+		#if WITH_EDITORONLY_DATA
+		// Pow22 only affects texture *Source* encoding of import
+		// PlatformData is always sRGB even if bUseLegacyGamma was on
+		if ( SRGB && bUseLegacyGamma )
+		{
+			return EGammaSpace::Pow22;
+		}
+		#endif
+
+		return SRGB ? EGammaSpace::sRGB : EGammaSpace::Linear;
 	}
-#endif
 
 	/** @return the width of the surface represented by the texture. */
 	virtual float GetSurfaceWidth() const PURE_VIRTUAL(UTexture::GetSurfaceWidth,return 0;);
diff --git a/Engine/Source/Runtime/Engine/Classes/Engine/Texture2D.h b/Engine/Source/Runtime/Engine/Classes/Engine/Texture2D.h
--- a/Engine/Source/Runtime/Engine/Classes/Engine/Texture2D.h
+++ b/Engine/Source/Runtime/Engine/Classes/Engine/Texture2D.h
@@ -119,6 +119,7 @@ public:
 	//~ End UObject Interface.
 
 	//~ Begin UTexture Interface.
+	virtual ETextureClass GetTextureClass() const override { return ETextureClass::TwoD; }
 	virtual float GetSurfaceWidth() const override { return static_cast<float>(GetSizeX()); }
 	virtual float GetSurfaceHeight() const override { return static_cast<float>(GetSizeY()); }
 	virtual float GetSurfaceDepth() const override { return 0.0f; }
diff --git a/Engine/Source/Runtime/Engine/Classes/Engine/Texture2DArray.cpp b/Engine/Source/Runtime/Engine/Classes/Engine/Texture2DArray.cpp
--- a/Engine/Source/Runtime/Engine/Classes/Engine/Texture2DArray.cpp
+++ b/Engine/Source/Runtime/Engine/Classes/Engine/Texture2DArray.cpp
@@ -223,7 +223,7 @@ ENGINE_API bool UTexture2DArray::UpdateSourceFromSourceTextures(bool bCreatingNe
 			Source.UnlockMip(MipIndex);
 		}
 
-		UpdateMipGenSettings();
+		ValidateSettingsAfterImportOrEdit();
 		SetLightingGuid();
 		UpdateResource();
 	}
@@ -317,32 +317,10 @@ uint32 UTexture2DArray::GetMaximumDimension() const
 
 }
 
-void UTexture2DArray::UpdateMipGenSettings()
-{
-	if (PowerOfTwoMode == ETexturePowerOfTwoSetting::None && !Source.IsPowerOfTwo())
-	{
-		// Force NPT textures to have no mip maps.
-		MipGenSettings = TMGS_NoMipmaps;
-		NeverStream = true;
-	}
-}
-
 ENGINE_API void UTexture2DArray::PostEditChangeProperty(FPropertyChangedEvent & PropertyChangedEvent)
 {	
 	const FName PropertyName = PropertyChangedEvent.GetPropertyName();
-
-	if (PowerOfTwoMode == ETexturePowerOfTwoSetting::None && (!Source.IsPowerOfTwo()))
-	{
-		// Force NPT textures to have no mip maps.
-		if (PropertyName == GET_MEMBER_NAME_CHECKED(UTexture2DArray, MipGenSettings)) 
-		{
-			UE_LOG(LogTexture, Warning, TEXT("Cannot use mip maps for non-power of two textures."));
-		}
-
-		MipGenSettings = TMGS_NoMipmaps;
-		NeverStream = true;
-	}
-
+	
 	if (PropertyName == GET_MEMBER_NAME_CHECKED(UTexture2DArray, SourceTextures))
 	{
 		// Empty SourceTextures, remove any resources if present.
diff --git a/Engine/Source/Runtime/Engine/Classes/Engine/Texture2DArray.h b/Engine/Source/Runtime/Engine/Classes/Engine/Texture2DArray.h
--- a/Engine/Source/Runtime/Engine/Classes/Engine/Texture2DArray.h
+++ b/Engine/Source/Runtime/Engine/Classes/Engine/Texture2DArray.h
@@ -48,6 +48,7 @@ public:
 	}
 
 	//~ Begin UTexture Interface
+	virtual ETextureClass GetTextureClass() const override { return ETextureClass::Array; }
 	virtual void Serialize(FArchive& Ar) override;
 	virtual void PostLoad() override;
 	virtual void GetAssetRegistryTags(TArray<FAssetRegistryTag>& OutTags) const override;
@@ -124,7 +125,6 @@ public:
 protected:
 
 #if WITH_EDITOR
-	void UpdateMipGenSettings();
 	virtual bool GetStreamableRenderResourceState(FTexturePlatformData* InPlatformData, FStreamableRenderResourceState& OutState) const override;
 #endif
 };
\ No newline at end of file
diff --git a/Engine/Source/Runtime/Engine/Classes/Engine/Texture2DDynamic.h b/Engine/Source/Runtime/Engine/Classes/Engine/Texture2DDynamic.h
--- a/Engine/Source/Runtime/Engine/Classes/Engine/Texture2DDynamic.h
+++ b/Engine/Source/Runtime/Engine/Classes/Engine/Texture2DDynamic.h
@@ -28,6 +28,8 @@ struct FTexture2DDynamicCreateInfo
 	ESamplerAddressMode SamplerAddressMode;
 };
 
+// note : UTexture2DDynamic derives directly from UTexture not from UTexture2D
+// UTexture2DDynamic is a base for textures that don't have a TextureSource
 UCLASS(hidecategories=Object, MinimalAPI)
 class UTexture2DDynamic : public UTexture
 {
@@ -54,6 +56,7 @@ class UTexture2DDynamic : public UTexture
 	
 public:
 	//~ Begin UTexture Interface.
+	ENGINE_API virtual ETextureClass GetTextureClass() const { return ETextureClass::TwoDDynamic; }
 	ENGINE_API virtual FTextureResource* CreateResource() override;
 	virtual EMaterialValueType GetMaterialType() const override { return MCT_Texture2D; }
 	ENGINE_API virtual float GetSurfaceWidth() const override;
diff --git a/Engine/Source/Runtime/Engine/Classes/Engine/TextureCube.h b/Engine/Source/Runtime/Engine/Classes/Engine/TextureCube.h
--- a/Engine/Source/Runtime/Engine/Classes/Engine/TextureCube.h
+++ b/Engine/Source/Runtime/Engine/Classes/Engine/TextureCube.h
@@ -65,6 +65,7 @@ public:
 	ENGINE_API void GetMipData(int32 FirstMipToLoad, void** OutMipData);
 	
 	//~ Begin UTexture Interface
+	virtual ETextureClass GetTextureClass() const override { return ETextureClass::Cube; }
 	virtual float GetSurfaceWidth() const override { return (float)GetSizeX(); }
 	virtual float GetSurfaceHeight() const override { return (float)GetSizeY(); }
 	virtual float GetSurfaceDepth() const override { return 0; }
diff --git a/Engine/Source/Runtime/Engine/Classes/Engine/TextureCubeArray.h b/Engine/Source/Runtime/Engine/Classes/Engine/TextureCubeArray.h
--- a/Engine/Source/Runtime/Engine/Classes/Engine/TextureCubeArray.h
+++ b/Engine/Source/Runtime/Engine/Classes/Engine/TextureCubeArray.h
@@ -44,6 +44,7 @@ public:
 	}
 
 	//~ Begin UTexture Interface
+	virtual ETextureClass GetTextureClass() const override { return ETextureClass::CubeArray; }
 	virtual void Serialize(FArchive& Ar) override;
 	virtual void PostLoad() override;
 	virtual void GetAssetRegistryTags(TArray<FAssetRegistryTag>& OutTags) const override;
diff --git a/Engine/Source/Runtime/Engine/Classes/Engine/TextureDefines.h b/Engine/Source/Runtime/Engine/Classes/Engine/TextureDefines.h
--- a/Engine/Source/Runtime/Engine/Classes/Engine/TextureDefines.h
+++ b/Engine/Source/Runtime/Engine/Classes/Engine/TextureDefines.h
@@ -244,6 +244,26 @@ enum class ETextureEncodeSpeed : uint8
 	Fast = 2
 };
 
+UENUM()
+enum class ETextureClass : uint8
+{
+	Invalid,
+	// Engine types with source data :
+	TwoD,
+	Cube,
+	Array,
+	CubeArray,
+	Volume,
+	
+	// Engine types without source data :
+	TwoDDynamic,
+	RenderTarget, // can be 2D or Cube
+
+	// User types :
+	Other2DNoSource, // Media, Web, etc. that should have derived from TwoDDynamic but didn't
+	OtherUnknown
+};
+
 UENUM()
 enum ECompositeTextureMode
 {
diff --git a/Engine/Source/Runtime/Engine/Classes/Engine/TextureRenderTarget.h b/Engine/Source/Runtime/Engine/Classes/Engine/TextureRenderTarget.h
--- a/Engine/Source/Runtime/Engine/Classes/Engine/TextureRenderTarget.h
+++ b/Engine/Source/Runtime/Engine/Classes/Engine/TextureRenderTarget.h
@@ -44,6 +44,7 @@ class UTextureRenderTarget : public UTexture
 
 
 	//~ Begin UTexture Interface
+	ENGINE_API virtual ETextureClass GetTextureClass() const { return ETextureClass::RenderTarget; }
 	ENGINE_API virtual class FTextureResource* CreateResource() override;
 	ENGINE_API virtual EMaterialValueType GetMaterialType() const override;
 	//~ End UTexture Interface
diff --git a/Engine/Source/Runtime/Engine/Classes/Engine/VolumeTexture.h b/Engine/Source/Runtime/Engine/Classes/Engine/VolumeTexture.h
--- a/Engine/Source/Runtime/Engine/Classes/Engine/VolumeTexture.h
+++ b/Engine/Source/Runtime/Engine/Classes/Engine/VolumeTexture.h
@@ -82,6 +82,7 @@ public:
 	}
 
 	//~ Begin UTexture Interface
+	virtual ETextureClass GetTextureClass() const override { return ETextureClass::Volume; }
 	virtual float GetSurfaceWidth() const override { return static_cast<float>(GetSizeX()); }
 	virtual float GetSurfaceHeight() const override { return static_cast<float>(GetSizeY()); }
 	virtual float GetSurfaceDepth() const override { return static_cast<float>(GetSizeZ()); }
diff --git a/Engine/Source/Runtime/Engine/Private/Texture.cpp b/Engine/Source/Runtime/Engine/Private/Texture.cpp
--- a/Engine/Source/Runtime/Engine/Private/Texture.cpp
+++ b/Engine/Source/Runtime/Engine/Private/Texture.cpp
@@ -363,12 +363,77 @@ bool UTexture::CanEditChange(const FProperty* InProperty) const
 	return true;
 }
 
+// we're in WITH_EDITOR but not sure that's right, maybe move out?
+void UTexture::ValidateSettingsAfterImportOrEdit()
+{	
+	// calling ValidateSettingsAfterImportOrEdit if all settings are already valid should be a nop
+	// if you call ValidateSettingsAfterImportOrEdit twice, the second should do nothing
+
+	// this will be called by PostEditChange() with no arg
+		
+#if WITH_EDITORONLY_DATA
+
+	bool IsPowerOfTwo = Source.IsPowerOfTwo();
+	if ( GetTextureClass() == ETextureClass::Volume && ! FMath::IsPowerOfTwo(Source.NumSlices) )
+	{
+		IsPowerOfTwo = false;
+
+		if ( PowerOfTwoMode != ETexturePowerOfTwoSetting::None )
+		{
+			UE_LOG(LogTexture, Warning, TEXT("PowerOfTwoMode not supported for volume Z"));
+			PowerOfTwoMode = ETexturePowerOfTwoSetting::None;
+		}
+	}
+
+	if (!IsPowerOfTwo && (PowerOfTwoMode == ETexturePowerOfTwoSetting::None))
+	{
+		// Force NPT textures to have no mipmaps.
+		MipGenSettings = TMGS_NoMipmaps;
+		NeverStream = true;
+		if (VirtualTextureStreaming)
+		{
+			UE_LOG(LogTexture, Warning, TEXT("VirtualTextureStreaming not supported for \"%s\", texture size is not a power-of-2"), *GetName());
+			VirtualTextureStreaming = false;
+		}
+	}
+
+	// Make sure settings are correct for LUT textures.
+	if(LODGroup == TEXTUREGROUP_ColorLookupTable)
+	{
+		MipGenSettings = TMGS_NoMipmaps;
+		SRGB = false;
+	}
+#endif // #if WITH_EDITORONLY_DATA
+
+	const bool bPreventSRGB = (CompressionSettings == TC_Alpha || CompressionSettings == TC_Normalmap || CompressionSettings == TC_Masks || CompressionSettings == TC_HDR || CompressionSettings == TC_HDR_Compressed || CompressionSettings == TC_HalfFloat);
+	if (bPreventSRGB && SRGB == true)
+	{
+		SRGB = false;
+	}
+
+#if WITH_EDITORONLY_DATA
+	if (MaxTextureSize <= 0)
+	{
+		MaxTextureSize = 0;
+	}
+	else
+	{
+		MaxTextureSize = FMath::Min<int32>(FMath::RoundUpToPowerOfTwo(MaxTextureSize), GetMaximumDimension());
+	}
+#endif
+	
+	NumCinematicMipLevels = FMath::Max<int32>( NumCinematicMipLevels, 0 );
+
+}
+
 void UTexture::PostEditChangeProperty(FPropertyChangedEvent& PropertyChangedEvent)
 {
 	TRACE_CPUPROFILER_EVENT_SCOPE(UTexture_PostEditChangeProperty);
 
 	Super::PostEditChangeProperty(PropertyChangedEvent);
 
+	// this whole function is in WITH_EDITOR so these ifs don't make much sense
+
 #if WITH_EDITOR
 	ON_SCOPE_EXIT
 	{
@@ -385,8 +450,11 @@ void UTexture::PostEditChangeProperty(FPropertyChangedEvent& PropertyChangedEven
 	}
 #endif // WITH_EDITOR
 
+	// assume there was a change that needs a new lighting guid :
 	SetLightingGuid();
 
+	ValidateSettingsAfterImportOrEdit();
+
 	// Determine whether any property that requires recompression of the texture, or notification to Materials has changed.
 	bool RequiresNotifyMaterials = false;
 	bool DeferCompressionWasEnabled = false;
@@ -402,7 +470,6 @@ void UTexture::PostEditChangeProperty(FPropertyChangedEvent& PropertyChangedEven
 		static const FName VirtualTextureStreamingName = GET_MEMBER_NAME_CHECKED(UTexture, VirtualTextureStreaming);
 #if WITH_EDITORONLY_DATA
 		static const FName SourceColorSpaceName = GET_MEMBER_NAME_CHECKED(FTextureSourceColorSettings, ColorSpace);
-		static const FName MaxTextureSizeName = GET_MEMBER_NAME_CHECKED(UTexture, MaxTextureSize);
 		static const FName CompressionQualityName = GET_MEMBER_NAME_CHECKED(UTexture, CompressionQuality);
 		static const FName OodleTextureSdkVersionName = GET_MEMBER_NAME_CHECKED(UTexture, OodleTextureSdkVersion);
 #endif //WITH_EDITORONLY_DATA
@@ -414,9 +481,10 @@ void UTexture::PostEditChangeProperty(FPropertyChangedEvent& PropertyChangedEven
 			(PropertyName == SrgbName))
 		{
 			RequiresNotifyMaterials = true;
-
+			
 			if (PropertyName == LODGroupName)
 			{
+				// should this be in Validate ? or only when switching to this LODGroup ? (allowing change after)
 				if (LODGroup == TEXTUREGROUP_8BitData)
 				{
 					CompressionSettings = TC_VectorDisplacementmap;
@@ -452,17 +520,6 @@ void UTexture::PostEditChangeProperty(FPropertyChangedEvent& PropertyChangedEven
 			RequiresNotifyMaterials = true;
 			bInvalidatesMaterialShaders = false;
 		}
-		else if (PropertyName == MaxTextureSizeName)
-		{
-			if (MaxTextureSize <= 0)
-			{
-				MaxTextureSize = 0;
-			}
-			else
-			{
-				MaxTextureSize = FMath::Min<int32>(FMath::RoundUpToPowerOfTwo(MaxTextureSize), GetMaximumDimension());
-			}
-		}
 		else if (PropertyName == VirtualTextureStreamingName)
 		{
 			RequiresNotifyMaterials = true;
@@ -480,11 +537,6 @@ void UTexture::PostEditChangeProperty(FPropertyChangedEvent& PropertyChangedEven
 #endif //WITH_EDITORONLY_DATA
 	}
 
-	const bool bPreventSRGB = (CompressionSettings == TC_Alpha || CompressionSettings == TC_Normalmap || CompressionSettings == TC_Masks || CompressionSettings == TC_HDR || CompressionSettings == TC_HDR_Compressed || CompressionSettings == TC_HalfFloat);
-	if (bPreventSRGB && SRGB == true)
-	{
-		SRGB = false;
-	}
 
 	if (!PropertyThatChanged && !GDisableAutomaticTextureMaterialUpdateDependencies)
 	{
@@ -523,8 +575,6 @@ void UTexture::PostEditChangeProperty(FPropertyChangedEvent& PropertyChangedEven
 		}
 	}
 
-	NumCinematicMipLevels = FMath::Max<int32>( NumCinematicMipLevels, 0 );
-
 	// Don't update the texture resource if we've turned "DeferCompression" on, as this 
 	// would cause it to immediately update as an uncompressed texture
 	if( !DeferCompressionWasEnabled && (PropertyChangedEvent.ChangeType & EPropertyChangeType::Interactive) == 0 )
diff --git a/Engine/Source/Runtime/Engine/Private/Texture2D.cpp b/Engine/Source/Runtime/Engine/Private/Texture2D.cpp
--- a/Engine/Source/Runtime/Engine/Private/Texture2D.cpp
+++ b/Engine/Source/Runtime/Engine/Private/Texture2D.cpp
@@ -491,27 +491,6 @@ void UTexture2D::PostEditUndo()
 
 void UTexture2D::PostEditChangeProperty(FPropertyChangedEvent& PropertyChangedEvent)
 {
-#if WITH_EDITORONLY_DATA
-	if (!Source.IsPowerOfTwo() && (PowerOfTwoMode == ETexturePowerOfTwoSetting::None))
-	{
-		// Force NPT textures to have no mipmaps.
-		MipGenSettings = TMGS_NoMipmaps;
-		NeverStream = true;
-		if (VirtualTextureStreaming)
-		{
-			UE_LOG(LogTexture, Warning, TEXT("VirtualTextureStreaming not supported for \"%s\", texture size is not a power-of-2"), *GetName());
-			VirtualTextureStreaming = false;
-		}
-	}
-
-	// Make sure settings are correct for LUT textures.
-	if(LODGroup == TEXTUREGROUP_ColorLookupTable)
-	{
-		MipGenSettings = TMGS_NoMipmaps;
-		SRGB = false;
-	}
-#endif // #if WITH_EDITORONLY_DATA
-
 	if (PropertyChangedEvent.GetPropertyName() == GET_MEMBER_NAME_CHECKED(UTexture2D, AddressX)
 		|| PropertyChangedEvent.GetPropertyName() == GET_MEMBER_NAME_CHECKED(UTexture2D, AddressY))
 	{
@@ -528,6 +507,9 @@ float UTexture2D::GetAverageBrightness(bool bIgnoreTrueBlack, bool bUseGrayscale
 {
 	float AvgBrightness = -1.0f;
 #if WITH_EDITOR
+
+	// @@!! this is done in a few places ; factor out to an FImage helper?
+
 	TArray64<uint8> RawData;
 	// use the source art if it exists
 	if (Source.IsValid() && Source.GetFormat() == TSF_BGRA8)
@@ -545,8 +527,7 @@ float UTexture2D::GetAverageBrightness(bool bIgnoreTrueBlack, bool bUseGrayscale
 		int32 SizeY = Source.GetSizeY();
 		double PixelSum = 0.0f;
 		int32 Divisor = SizeX * SizeY;
-		FColor* ColorData = (FColor*)RawData.GetData();
-		FLinearColor CurrentColor;
+		const FColor* ColorData = (const FColor*)RawData.GetData();
 		for (int32 Y = 0; Y < SizeY; Y++)
 		{
 			for (int32 X = 0; X < SizeX; X++)
@@ -558,6 +539,7 @@ float UTexture2D::GetAverageBrightness(bool bIgnoreTrueBlack, bool bUseGrayscale
 					continue;
 				}
 
+				FLinearColor CurrentColor;
 				if (SRGB == true)
 				{
 					CurrentColor = bUseLegacyGamma ? FLinearColor::FromPow22Color(*ColorData) : FLinearColor(*ColorData);
diff --git a/Engine/Source/Runtime/Engine/Private/TextureCubeArray.cpp b/Engine/Source/Runtime/Engine/Private/TextureCubeArray.cpp
--- a/Engine/Source/Runtime/Engine/Private/TextureCubeArray.cpp
+++ b/Engine/Source/Runtime/Engine/Private/TextureCubeArray.cpp
@@ -451,6 +451,9 @@ ENGINE_API bool UTextureCubeArray::UpdateSourceFromSourceTextures(bool bCreating
 	Source.Init(SizeX, SizeY, ArraySize * NumSlices, NumMips, Format);
 	Source.bLongLatCubemap = (NumSlices == 1);
 
+	// this path sets bLongLatCubemap for CubeArray
+	//  most paths do not, so it is not reliable
+
 	// We only copy the top level Mip map.
 	TArray<uint8*, TInlineAllocator<MAX_TEXTURE_MIP_COUNT> > DestMipData;
 	TArray<uint64, TInlineAllocator<MAX_TEXTURE_MIP_COUNT> > MipSizeBytes;
@@ -485,6 +488,7 @@ ENGINE_API bool UTextureCubeArray::UpdateSourceFromSourceTextures(bool bCreating
 	}
 
 	SetLightingGuid();
+	ValidateSettingsAfterImportOrEdit();
 	UpdateResource();
 
 	return true;
diff --git a/Engine/Source/Runtime/Engine/Private/TextureDerivedData.cpp b/Engine/Source/Runtime/Engine/Private/TextureDerivedData.cpp
--- a/Engine/Source/Runtime/Engine/Private/TextureDerivedData.cpp
+++ b/Engine/Source/Runtime/Engine/Private/TextureDerivedData.cpp
@@ -733,28 +733,47 @@ static void GetTextureBuildSettings(
 		OutBuildSettings.MaxTextureResolution = Texture.MaxTextureSize;
 	}
 
-	if (Texture.IsA(UTextureCube::StaticClass()))
+	ETextureClass TextureClass = Texture.GetTextureClass();
+	
+	if ( TextureClass == ETextureClass::TwoD )
+	{
+		// nada
+	}
+	else if ( TextureClass == ETextureClass::Cube )
 	{
 		OutBuildSettings.bCubemap = true;
 		OutBuildSettings.DiffuseConvolveMipLevel = GDiffuseConvolveMipLevel;
+		// beware IsLongLatCubemap is pretty useless , GetNumSlices() is doing the work here
 		check( Texture.Source.GetNumSlices() == 1 || Texture.Source.GetNumSlices() == 6 );
 		OutBuildSettings.bLongLatSource = (Texture.Source.GetNumSlices() == 1) || Texture.Source.IsLongLatCubemap();
 	}
-	else if (Texture.IsA(UTexture2DArray::StaticClass()))
+	else if ( TextureClass == ETextureClass::Array )
 	{
 		OutBuildSettings.bTextureArray = true;
 	}
-	else if (Texture.IsA(UTextureCubeArray::StaticClass()))
+	else if ( TextureClass == ETextureClass::CubeArray )
 	{
 		OutBuildSettings.bCubemap = true;
 		OutBuildSettings.bTextureArray = true;
+		// beware IsLongLatCubemap
+		// ambiguous with longlat cube arrays with multiple of 6 array size
 		OutBuildSettings.bLongLatSource = Texture.Source.IsLongLatCubemap();
 		check( ((Texture.Source.GetNumSlices()%6)==0) || OutBuildSettings.bLongLatSource );
 	}
-	else if (Texture.IsA(UVolumeTexture::StaticClass()))
+	else if ( TextureClass == ETextureClass::Volume )
 	{
 		OutBuildSettings.bVolume = true;
 	}
+	else if ( TextureClass == ETextureClass::TwoDDynamic ||
+		TextureClass == ETextureClass::Other2DNoSource )
+	{
+		UE_LOG(LogTexture, Warning, TEXT("Unexpected texture build for dynamic texture? (%s)"),*Texture.GetName());
+	}
+	else
+	{
+		// unknown TextureType ?
+		UE_LOG(LogTexture, Error, TEXT("Unexpected texture build for unknown texture class? (%s)"),*Texture.GetName());
+	}
 
 	bool bDownsampleWithAverage;
 	bool bSharpenWithoutColorShift;
diff --git a/Engine/Source/Runtime/Engine/Private/VolumeTexture.cpp b/Engine/Source/Runtime/Engine/Private/VolumeTexture.cpp
--- a/Engine/Source/Runtime/Engine/Private/VolumeTexture.cpp
+++ b/Engine/Source/Runtime/Engine/Private/VolumeTexture.cpp
@@ -100,8 +100,8 @@ bool UVolumeTexture::UpdateSourceFromSourceTexture()
 			*PlatformData = FTexturePlatformData();
 		}
 	}
-
-	UpdateMipGenSettings();
+	
+	ValidateSettingsAfterImportOrEdit();
 #endif // WITH_EDITOR
 
 	return bSourceValid;
@@ -155,7 +155,7 @@ ENGINE_API bool UVolumeTexture::UpdateSourceFromFunction(TFunction<void(int32, i
 
 	SetLightingGuid(); // Because the content has changed, use a new GUID.
 
-	UpdateMipGenSettings();
+	ValidateSettingsAfterImportOrEdit();
 
 	// Make sure to update the texture resource so the results of filling the texture 
 	UpdateResource();
@@ -369,8 +369,6 @@ void UVolumeTexture::PostEditChangeProperty(FPropertyChangedEvent& PropertyChang
 		}
 	}
 	
-	UpdateMipGenSettings();
-
 	Super::PostEditChangeProperty(PropertyChangedEvent);
 }
 
@@ -380,16 +378,6 @@ uint32 UVolumeTexture::GetMaximumDimension() const
 	return GetMax2DTextureDimension();
 }
 
-void UVolumeTexture::UpdateMipGenSettings()
-{
-	if (PowerOfTwoMode == ETexturePowerOfTwoSetting::None && (!Source.IsPowerOfTwo() || !FMath::IsPowerOfTwo(Source.NumSlices)))
-	{
-		// Force NPT textures to have no mipmaps.
-		MipGenSettings = TMGS_NoMipmaps;
-		NeverStream = true;
-	}
-}
-
 #endif // #if WITH_EDITOR
 
 bool UVolumeTexture::StreamOut(int32 NewMipCount)
diff --git a/Engine/Source/Runtime/MediaAssets/Public/MediaTexture.h b/Engine/Source/Runtime/MediaAssets/Public/MediaTexture.h
--- a/Engine/Source/Runtime/MediaAssets/Public/MediaTexture.h
+++ b/Engine/Source/Runtime/MediaAssets/Public/MediaTexture.h
@@ -41,6 +41,9 @@ enum MediaTextureOrientation
 
 /**
  * Implements a texture asset for rendering video tracks from UMediaPlayer assets.
+ * 
+ * note: derives directly from UTexture, not from UTexture2D or UTexture2DDynamic
+ *    maybe should have been UTexture2DDynamic?
  */
 UCLASS(hidecategories=(Adjustments, Compositing, LevelOfDetail, ImportSettings, Object))
 class MEDIAASSETS_API UMediaTexture
@@ -206,6 +209,7 @@ public:
 	virtual uint32 GetSurfaceArraySize() const override { return 0; }
 	virtual FGuid GetExternalTextureGuid() const override;
 	void SetRenderedExternalTextureGuid(const FGuid& InNewGuid);
+	virtual ETextureClass GetTextureClass() const { return ETextureClass::Other2DNoSource; }
 
 public:
 
diff --git a/Engine/Source/Runtime/WebBrowserTexture/Public/WebBrowserTexture.h b/Engine/Source/Runtime/WebBrowserTexture/Public/WebBrowserTexture.h
--- a/Engine/Source/Runtime/WebBrowserTexture/Public/WebBrowserTexture.h
+++ b/Engine/Source/Runtime/WebBrowserTexture/Public/WebBrowserTexture.h
@@ -17,6 +17,8 @@
 
 /**
 * Implements a texture asset for rendering webbrowser output for Android.
+* 
+*  probably should have derived from UTexture2DDynamic (not UTexture)
 */
 UCLASS(hidecategories = (Adjustments, Compositing, LevelOfDetail, Object))
 class WEBBROWSERTEXTURE_API UWebBrowserTexture
@@ -72,6 +74,7 @@ public:
 	virtual float GetSurfaceDepth() const override { return 0; }
 	virtual uint32 GetSurfaceArraySize() const override { return 0; }
 	virtual FGuid GetExternalTextureGuid() const override;
+	virtual ETextureClass GetTextureClass() const { return ETextureClass::Other2DNoSource; }
 
 public:
 	//~ UObject interface.
