fix mip map gen on non power of 2 textures
also allow MaxTextureSize on non-pow2
no longer force TMGS_NoMipMaps on non-pow2

diff --git a/Engine/Source/Developer/TextureCompressor/Private/TextureCompressorModule.cpp b/Engine/Source/Developer/TextureCompressor/Private/TextureCompressorModule.cpp
--- a/Engine/Source/Developer/TextureCompressor/Private/TextureCompressorModule.cpp
+++ b/Engine/Source/Developer/TextureCompressor/Private/TextureCompressorModule.cpp
@@ -25,6 +25,8 @@ DEFINE_LOG_CATEGORY_STATIC(LogTextureCompressor, Log, All);
 	Mip-Map Generation
 ------------------------------------------------------------------------------*/
 
+// NOTE: mip gen wrap/clamp does NOT correspond to Texture Address Wrap/Clamp setting !!
+//		it comes from bPreserveBorder !
 enum EMipGenAddressMode
 {
 	MGTAM_Wrap,
@@ -51,6 +53,7 @@ struct FImageView2D
 	{
 		SizeX = Image.SizeX;
 		SizeY = Image.SizeY;
+		check( SliceIndex < Image.NumSlices );
 		SliceColors = (&Image.AsRGBA32F()[0]) + SliceIndex * SizeY * SizeX;
 	}
 
@@ -83,6 +86,10 @@ static const FLinearColor& LookupSourceMip(const FImageView2D& SourceImageData,
 	if(AddressMode == MGTAM_Wrap)
 	{
 		// wrap
+		// ! requires pow2 sizes
+		checkSlow( FMath::IsPowerOfTwo(SourceImageData.SizeX) );
+		checkSlow( FMath::IsPowerOfTwo(SourceImageData.SizeY) );
+
 		X = (int32)((uint32)X) & (SourceImageData.SizeX - 1);
 		Y = (int32)((uint32)Y) & (SourceImageData.SizeY - 1);
 	}
@@ -752,8 +759,9 @@ static void GenerateSharpenedMipB8G8R8A8Templ(
 	bool bSharpenWithoutColorShift,
 	bool bUnfiltered)
 {
-	check( SourceImageData.SizeX == ScaleFactor * DestImageData.SizeX || DestImageData.SizeX == 1 );
-	check( SourceImageData.SizeY == ScaleFactor * DestImageData.SizeY || DestImageData.SizeY == 1 );
+	check( ScaleFactor == 1 || ScaleFactor == 2 );
+	check( (SourceImageData.SizeX/ScaleFactor) == DestImageData.SizeX || DestImageData.SizeX == 1 );
+	check( (SourceImageData.SizeY/ScaleFactor) == DestImageData.SizeY || DestImageData.SizeY == 1 );
 	
 	int32 KernelFilterTableSize = (int32) Kernel.GetFilterTableSize();
 
@@ -887,8 +895,11 @@ static void GenerateMipBorder(
 	FImageView2D& DestImageData
 	)
 {
-	check( SrcImageData.SizeX == 2 * DestImageData.SizeX || DestImageData.SizeX == 1 );
-	check( SrcImageData.SizeY == 2 * DestImageData.SizeY || DestImageData.SizeY == 1 );
+	check( (SrcImageData.SizeX/2) == DestImageData.SizeX || DestImageData.SizeX == 1 );
+	check( (SrcImageData.SizeY/2) == DestImageData.SizeY || DestImageData.SizeY == 1 );
+
+	// this check is unnecessary if we always used MGTAM_Clamp here;
+	bool bIsPow2 = FMath::IsPowerOfTwo(SrcImageData.SizeX) && FMath::IsPowerOfTwo(SrcImageData.SizeY);
 
 	for ( int32 DestY = 0; DestY < DestImageData.SizeY; DestY++ )
 	{
@@ -910,7 +921,17 @@ static void GenerateMipBorder(
 							SourceY == 0 ||
 							SourceY == SrcImageData.SizeY - 1 )
 						{
-							FLinearColor Sample = LookupSourceMip<MGTAM_Wrap>( SrcImageData, SourceX, SourceY );
+							// I think this should have just always been MGTAM_Clamp
+							//	but that changes existing content, so preserve old behavior of using _Wrap :(
+							FLinearColor Sample;
+							if ( bIsPow2 )
+							{
+								Sample = LookupSourceMip<MGTAM_Wrap>( SrcImageData, SourceX, SourceY );
+							}
+							else
+							{
+								Sample = LookupSourceMip<MGTAM_Clamp>( SrcImageData, SourceX, SourceY );
+							}
 							FilteredColor += Sample;
 							WeightSum += 1.0f;
 						}
@@ -939,10 +960,26 @@ static void GenerateMipBorder(
 }
 
 // how should be treat lookups outside of the image
-static EMipGenAddressMode ComputeAdressMode(const FTextureBuildSettings& Settings)
+static EMipGenAddressMode ComputeAddressMode(const FImage & Image,const FTextureBuildSettings& Settings)
 {
+	// note: all textures Wrap by default even if their address mode is set to Clamp !?
 	EMipGenAddressMode AddressMode = MGTAM_Wrap;
 
+	// Wrap uses AND so requires pow2 sizes ; change to Clamp if nonpow2
+	bool bIsPow2 = FMath::IsPowerOfTwo(Image.SizeX) && FMath::IsPowerOfTwo(Image.SizeY);
+	// 2d address mode, no need to look at volume z :
+	/*
+	if ( Settings.bVolume && ! FMath::IsPowerOfTwo(Image.NumSlices) )
+	{
+		bIsPow2 = false;
+	}
+	*/
+
+	if ( ! bIsPow2 )
+	{
+		AddressMode = MGTAM_Clamp;
+	}
+
 	if(Settings.bPreserveBorder)
 	{
 		AddressMode = Settings.bBorderColorBlack ? MGTAM_BorderBlack : MGTAM_Clamp;
@@ -958,7 +995,7 @@ static void GenerateTopMip(const FImage& SrcImage, FImage& DestImage, const FTex
 	// GenerateTopMip is only used for ApplyCompositeTexture
 	// bApplyKernelToTopMip is not exposed to Texture GUI
 
-	EMipGenAddressMode AddressMode = ComputeAdressMode(Settings);
+	EMipGenAddressMode AddressMode = ComputeAddressMode(SrcImage,Settings);
 
 	FImageKernel2D KernelDownsample;
 
@@ -1232,7 +1269,7 @@ void ITextureCompressorModule::GenerateMipChain(
 
 	//@TODO : add a true 3D kernel.
 
-	EMipGenAddressMode AddressMode = ComputeAdressMode(Settings);
+	EMipGenAddressMode AddressMode = ComputeAddressMode(*IntermediateSrcPtr,Settings);
 	bool bReDrawBorder = false;
 	if( Settings.bPreserveBorder )
 	{
@@ -1818,6 +1855,10 @@ void ITextureCompressorModule::GenerateAngularFilteredMips(TArray<FImage>& InOut
 	Exchange(SrcMipChain, InOutMipChain);
 	InOutMipChain.Empty(NumMips);
 
+	// note: should work on cube arrays but currently does not
+	// GetMipGenSettings forces Angular off for anything but pure Cube classes (no arrays)
+	check( SrcMipChain[0].NumSlices == 6 );
+
 	// Generate simple averaged mips to accelerate angular filtering.
 	for (int32 MipIndex = SrcMipChain.Num(); MipIndex < NumMips; ++MipIndex)
 	{
@@ -2599,6 +2640,16 @@ static void NormalizeMip(FImage& InOutMip)
 	}
 }
 
+static uint32 FloorLogTwo(uint32 X) 
+{
+	uint32 L2 = FMath::CeilLogTwo(X);
+	if ( ! FMath::IsPowerOfTwo(X) )
+	{
+		L2--;
+	}
+	return L2;
+}
+
 /**
  * Texture compression module
  */
@@ -2682,7 +2733,7 @@ public:
 		{
 			return false;
 		}
-
+		
 		// apply roughness adjustment depending on normal map variation
 		if (AssociatedNormalSourceMips.Num())
 		{
@@ -2762,6 +2813,7 @@ public:
 	// IModuleInterface implementation.
 	void StartupModule()
 	{
+		// @@!! what's this for ?
 #if PLATFORM_WINDOWS
 	#if PLATFORM_64BITS
 		if (FWindowsPlatformMisc::HasAVX2InstructionSupport())
@@ -2793,127 +2845,100 @@ private:
 #endif	//PLATFORM_WINDOWS
 
 	bool BuildTextureMips(
-		const TArray<FImage>& InSourceMips,
+		const TArray<FImage>& InSourceMipChain,
 		const FTextureBuildSettings& BuildSettings,
 		const FTextureFormatCompressorCaps& CompressorCaps,
 		const bool bNeedLinearize,
 		TArray<FImage>& OutMipChain)
 	{
 		TRACE_CPUPROFILER_EVENT_SCOPE(Texture.BuildTextureMips);
+		
+		// @@!! TEMP
+		//const_cast<FTextureBuildSettings&>(BuildSettings).MaxTextureResolution = 300;
 
-		check(InSourceMips.Num());
-		check(InSourceMips[0].SizeX > 0 && InSourceMips[0].SizeY > 0 && InSourceMips[0].NumSlices > 0);
+		check(InSourceMipChain.Num() > 0);
+		check(InSourceMipChain[0].SizeX > 0 && InSourceMipChain[0].SizeY > 0 && InSourceMipChain[0].NumSlices > 0);
 
 		// Identify long-lat cubemaps.
 		const bool bLongLatCubemap = BuildSettings.bLongLatSource;
 		if (BuildSettings.bCubemap && !bLongLatCubemap)
 		{
-			if (BuildSettings.bTextureArray && (InSourceMips[0].NumSlices % 6) != 0)
+			if (BuildSettings.bTextureArray && (InSourceMipChain[0].NumSlices % 6) != 0)
 			{
-				// Cube array must have multiiple of 6 slices
+				// Cube array must have multiple of 6 slices
 				return false;
 			}
-			if (!BuildSettings.bTextureArray && InSourceMips[0].NumSlices != 6)
+			if (!BuildSettings.bTextureArray && InSourceMipChain[0].NumSlices != 6)
 			{
 				// Non-array cube must have exactly 6 slices
 				return false;
 			}
 		}
 
-		// @@!! this is broken for non-pow2 sizes
-		//	if size is non-pow2, outer code should have set TMGS_NoMipmaps
-		
-		// problem 1 : this should be FloorLogTwo not Ceil
-		// problem 2 : if you then do that, it's broken for padding to pow2 mode, because the pad is after the mip count calculation
-		// problem 3 : pad+leave existing is broken
-
-		// Determine the maximum possible mip counts for source and dest.
-		const int32 MaxSourceMipCount = bLongLatCubemap ?
-			1 + FMath::CeilLogTwo(ComputeLongLatCubemapExtents(InSourceMips[0], BuildSettings.MaxTextureResolution)) :
-			1 + FMath::CeilLogTwo(FMath::Max3(InSourceMips[0].SizeX, InSourceMips[0].SizeY, BuildSettings.bVolume ? InSourceMips[0].NumSlices : 1));
-		const int32 MaxDestMipCount = 1 + FMath::CeilLogTwo(FMath::Min(CompressorCaps.MaxTextureDimension, BuildSettings.MaxTextureResolution));
-
-		// Determine the number of mips required by BuildSettings.
-		int32 NumOutputMips = (BuildSettings.MipGenSettings == TMGS_NoMipmaps) ? 1 : MaxSourceMipCount;
+		// handling of bLongLatCubemap seems overly complicated
+		//	what it should do is convert it right at the start here
+		//	then treat it as a standard cubemap below, no special cases
+		//	but that will change output :(
 
-		int32 NumSourceMips = InSourceMips.Num();
+		// pSourceMips will track the current FImages we consider to be "source"
+		const TArray<FImage> * pSourceMips = &InSourceMipChain;
 
-		// See if the smallest provided mip image is still too large for the current compressor.
-		int32 LevelsToUsableSource = FMath::Max(0, MaxSourceMipCount - MaxDestMipCount);
-		int32 StartMip = FMath::Max(0, LevelsToUsableSource);
-
-		if (BuildSettings.MipGenSettings == TMGS_LeaveExistingMips)
-		{
-			// @@!! if LeaveExistingMips and we want 10 mips but are only given 8, we don't make the lower ones?
-			//	 is that right?
-			int32 NumExistingMips = NumSourceMips - StartMip;
-			check( NumExistingMips <= NumOutputMips );
-
-			if (NumExistingMips <= 0)
-			{
-				// We can't generate 0 mip maps
-				UE_LOG(LogTextureCompressor, Warning,
-					TEXT("The source image has %d mips while the first mip would be %d. Please verify the maximum texture size or change the mips gen settings."),
-					NumSourceMips,
-					StartMip);
-				return false;
-			}
-			else
-			{
-				// if NumExistingMips < NumOutputMips
-				// a non-complete mip tail will be generated
-
-				NumOutputMips = NumExistingMips;
-			}
-		}
-
-		NumOutputMips = FMath::Min(NumOutputMips, MaxDestMipCount);
-
-
-		if (BuildSettings.MipGenSettings != TMGS_LeaveExistingMips || bLongLatCubemap)
-		{
-			NumSourceMips = 1;
-		}
+		// first pad up to pow2 if requested
+		ETexturePowerOfTwoSetting::Type PowerOfTwoMode = (ETexturePowerOfTwoSetting::Type) BuildSettings.PowerOfTwoMode;
 
+		// if bPadOrStretchTextureis done, PaddedSourceMips is filled with a new image
 		TArray<FImage> PaddedSourceMips;
 
+		if ( PowerOfTwoMode != ETexturePowerOfTwoSetting::None )
 		{
-			// @@!! Padded + leave existing is broken but that's not enforced ?
+			const FImage& FirstSourceMipImage = (*pSourceMips)[0];
 
-			const FImage& FirstSourceMipImage = InSourceMips[0];
 			int32 TargetTextureSizeX = FirstSourceMipImage.SizeX;
 			int32 TargetTextureSizeY = FirstSourceMipImage.SizeY;
 			int32 TargetTextureSizeZ = BuildSettings.bVolume ? FirstSourceMipImage.NumSlices : 1; // Only used for volume texture.
-			bool bPadOrStretchTexture = false;
 
 			const int32 PowerOfTwoTextureSizeX = FMath::RoundUpToPowerOfTwo(TargetTextureSizeX);
 			const int32 PowerOfTwoTextureSizeY = FMath::RoundUpToPowerOfTwo(TargetTextureSizeY);
 			const int32 PowerOfTwoTextureSizeZ = FMath::RoundUpToPowerOfTwo(TargetTextureSizeZ);
-			switch (static_cast<const ETexturePowerOfTwoSetting::Type>(BuildSettings.PowerOfTwoMode))
+
+			switch (PowerOfTwoMode)
 			{
-			case ETexturePowerOfTwoSetting::None:
-				break;
+			// None should not get here
 
 			case ETexturePowerOfTwoSetting::PadToPowerOfTwo:
-				bPadOrStretchTexture = true;
 				TargetTextureSizeX = PowerOfTwoTextureSizeX;
 				TargetTextureSizeY = PowerOfTwoTextureSizeY;
 				TargetTextureSizeZ = PowerOfTwoTextureSizeZ;
 				break;
 
 			case ETexturePowerOfTwoSetting::PadToSquarePowerOfTwo:
-				bPadOrStretchTexture = true;
-				TargetTextureSizeX = TargetTextureSizeY = FMath::Max3<int32>(PowerOfTwoTextureSizeX, PowerOfTwoTextureSizeY, PowerOfTwoTextureSizeZ);
-				// @@!! doesn't change or check TargetTextureSizeZ ?
+				TargetTextureSizeX = TargetTextureSizeY = TargetTextureSizeZ =
+					FMath::Max3<int32>(PowerOfTwoTextureSizeX, PowerOfTwoTextureSizeY, PowerOfTwoTextureSizeZ);
 				break;
 
 			default:
 				checkf(false, TEXT("Unknown entry in ETexturePowerOfTwoSetting::Type"));
 				break;
 			}
+			
+			bool bPadOrStretchTexture = 
+				(TargetTextureSizeX != FirstSourceMipImage.SizeX) ||
+				(TargetTextureSizeY != FirstSourceMipImage.SizeY) ||
+				( BuildSettings.bVolume && TargetTextureSizeZ != FirstSourceMipImage.NumSlices );
 
 			if (bPadOrStretchTexture)
 			{
+				if (BuildSettings.MipGenSettings == TMGS_LeaveExistingMips)
+				{
+					// pad+leave existing is broken
+					UE_LOG(LogTextureCompressor, Error,	TEXT("Texture padded to pow2 + LeaveExistingMips forbidden"));
+					return false;
+				}
+				if ( bLongLatCubemap )
+				{
+					UE_LOG(LogTextureCompressor, Warning, TEXT("PadPow2 + LongLat cubemap doesn't work, continuing.."));
+				}
+
 				// Want to stretch or pad the texture
 				bool bSuitableFormat = FirstSourceMipImage.Format == ERawImageFormat::RGBA32F;
 
@@ -2966,64 +2991,141 @@ private:
 						}
 					}
 				}
+				
+				// change pSourceMips to point at the one padded image we made
+				pSourceMips = &PaddedSourceMips;
 			}
-		}
+		}		
+
+		// now pow2 pad is done
+		// find a starting source that meets MaxTextureResolution limit
 
-		const TArray<FImage>& PostOptionalUpscaleSourceMips = (PaddedSourceMips.Num() > 0) ? PaddedSourceMips : InSourceMips;
+		int32 StartMip = 0;
 
-		bool bBuildSourceImage = StartMip > (NumSourceMips - 1);
+		TArray<FImage> BuildSourceImageMips;
 
-		TArray<FImage> GeneratedSourceMips;
-		if (bBuildSourceImage)
+		if ( ! bLongLatCubemap )
 		{
-			// should not get in here with LeaveExistingMips because NumExistingMips would be < 0
+			int32 NumSourceMips = (BuildSettings.MipGenSettings == TMGS_LeaveExistingMips) ? pSourceMips->Num() : 1;
+			
+			int64 MaxTextureResolution = BuildSettings.MaxTextureResolution;
 
-			// the source is larger than the compressor allows and no mip image exists to act as a smaller source.
-			// We must generate a suitable source image:
-			bool bSuitableFormat = PostOptionalUpscaleSourceMips.Last().Format == ERawImageFormat::RGBA32F;
-			const FImage& BaseImage = PostOptionalUpscaleSourceMips.Last();
+			// note that "LODBias" is very similar to MaxTextureResolution
+			//	but for LODBias we go ahead and make all the mips here
+			//	and then just don't serialize the top ones in TextureDerivedData
+			// (LODBias is not actually LOD Bias, it means discard top N mips)
+
+			// step through source mips to find one that meets MaxTextureResolution
+			while( StartMip < NumSourceMips && (
+				(*pSourceMips)[StartMip].SizeX > MaxTextureResolution ||
+				(*pSourceMips)[StartMip].SizeY > MaxTextureResolution ) )
+			{
+				StartMip++;
+			}
 
-			if (BaseImage.SizeX != FMath::RoundUpToPowerOfTwo(BaseImage.SizeX) || BaseImage.SizeY != FMath::RoundUpToPowerOfTwo(BaseImage.SizeY))
+			if ( StartMip == NumSourceMips )
 			{
-				UE_LOG(LogTextureCompressor, Warning,
-					TEXT("Source image %dx%d (npot) prevents resizing and is too large for compressors max dimension (%d)."),
+				if (BuildSettings.MipGenSettings == TMGS_LeaveExistingMips)
+				{
+					UE_LOG(LogTextureCompressor, Error,	TEXT("LeaveExistingMips no mip that fits max dimension (%d)."),(int)MaxTextureResolution);
+					return false;
+				}
+
+				// currently only makes mips if you only had 1 source mip :
+				check(NumSourceMips == 1);
+
+				// bLongLatCubemap should not get here because cube size is made from MaxTextureSize
+				check( ! bLongLatCubemap );
+
+				// the source is larger than the compressor allows and no mip image exists to act as a smaller source.
+				// We must generate a suitable source image:
+				const FImage& BaseImage = pSourceMips->Last();
+				bool bSuitableFormat = BaseImage.Format == ERawImageFormat::RGBA32F;
+			
+				check( MaxTextureResolution > 1 );
+				check( BaseImage.SizeX > MaxTextureResolution || 
+					   BaseImage.SizeY > MaxTextureResolution );
+
+				FImage Temp;
+				if (!bSuitableFormat)
+				{
+					// convert to RGBA32F
+					BaseImage.CopyTo(Temp, ERawImageFormat::RGBA32F, EGammaSpace::Linear);
+				}
+
+				UE_LOG(LogTextureCompressor, Verbose,
+					TEXT("Source image %dx%d too large for compressors max dimension (%d). Resizing."),
 					BaseImage.SizeX,
 					BaseImage.SizeY,
-					CompressorCaps.MaxTextureDimension
+					BuildSettings.MaxTextureResolution
 					);
-				return false;
-			}
 
-			FImage Temp;
-			if (!bSuitableFormat)
-			{
-				// convert to RGBA32F
-				BaseImage.CopyTo(Temp, ERawImageFormat::RGBA32F, EGammaSpace::Linear);
-			}
+				// Max Texture Size resizing happens here :
+				// note we do not check for TMGS_Angular here
+				GenerateMipChain(BuildSettings, bSuitableFormat ? BaseImage : Temp, BuildSourceImageMips, 1);
 
-			UE_LOG(LogTextureCompressor, Verbose,
-				TEXT("Source image %dx%d too large for compressors max dimension (%d). Resizing."),
-				BaseImage.SizeX,
-				BaseImage.SizeY,
-				CompressorCaps.MaxTextureDimension
-				);
-			// Max Texture Size resizing happens here :
-			GenerateMipChain(BuildSettings, bSuitableFormat ? BaseImage : Temp, GeneratedSourceMips, LevelsToUsableSource);
+				while( BuildSourceImageMips.Last().SizeX > MaxTextureResolution || 
+					   BuildSourceImageMips.Last().SizeY > MaxTextureResolution )
+				{
+					// note: now making mips one by one, rather than N in one call
+					//	this is not exactly the same if AlphaCoverage processing is on
+					GenerateMipChain(BuildSettings, BuildSourceImageMips.Last(), BuildSourceImageMips, 1);
+				}
+			
+				check( BuildSourceImageMips.Last().SizeX <= MaxTextureResolution &&
+					   BuildSourceImageMips.Last().SizeY <= MaxTextureResolution );
+				check( BuildSourceImageMips.Last().SizeX*2 > MaxTextureResolution ||
+					   BuildSourceImageMips.Last().SizeY*2 > MaxTextureResolution );
+
+				// change pSourceMips to point at the mip chain we made
+				pSourceMips = &BuildSourceImageMips;
+				StartMip = BuildSourceImageMips.Num() - 1;
+				// [StartMip] will now references BuildSourceImageMips.Last()
+			}
+		}
+				
+		// now shrinking to MaxTextureResolution is done, figure out which mips to use or make
+		
+		// Copy over base mip and any LeaveExisting, from SourceMips , starting at StartMip
+		int32 CopyCount = pSourceMips->Num() - StartMip;
+		check( CopyCount > 0 );
+		
+		int32 NumOutputMips;
 
-			check(GeneratedSourceMips.Num() != 0);
-			// Note: The newly generated mip chain does not include the original top level mip.
-			check( StartMip > 0 );
-			StartMip--;
+		if ( BuildSettings.MipGenSettings == TMGS_NoMipmaps )
+		{
+			NumOutputMips = 1;
+			CopyCount = 1;
+		}
+		else if ( BuildSettings.MipGenSettings == TMGS_LeaveExistingMips )
+		{
+			// only output what we can copy, generate none
+			NumOutputMips = CopyCount;
 		}
+		else
+		{
+			const FImage & TopMip = (*pSourceMips)[StartMip];
 
-		const TArray<FImage>& SourceMips = bBuildSourceImage ? GeneratedSourceMips : PostOptionalUpscaleSourceMips;
+			int32 TopMipSizeZ = BuildSettings.bVolume ? TopMip.NumSlices : 1;
+		
+			// NumOutputMips is the number of mips that would be made if you made a full mip chain
+			//  eg. 256 makes 9 mips , 300 also makes 9 mips
+			NumOutputMips = 1 + FloorLogTwo(
+				bLongLatCubemap ?
+				ComputeLongLatCubemapExtents(TopMip, BuildSettings.MaxTextureResolution) :
+				FMath::Max3(TopMip.SizeX, TopMip.SizeY, TopMipSizeZ) );
+				
+			// unless LeaveExistingMips, we only copy 1 
+			//  (in theory we could copy some existing and generate the rest, but that's not done currently)
+			//  (intentionally so, artists use this to limit mipping down)
+			CopyCount = 1;
+		}
 
+		// we will output NumOutputMips
 		OutMipChain.Empty(NumOutputMips);
-		// Copy over base mips.
-		check(StartMip < SourceMips.Num());
-		int32 CopyCount = SourceMips.Num() - StartMip;
 
 		int32 GenerateCount = NumOutputMips - CopyCount;
+		check( GenerateCount >= 0 );
 
 		// avoid converting to RGBA32F linear format if there's no need for any extra processing of pixels
 		// image will be left in BGRA8 format if possible
@@ -3035,7 +3137,10 @@ private:
 
 		for (int32 MipIndex = StartMip; MipIndex < StartMip + CopyCount; ++MipIndex)
 		{
-			const FImage& Image = SourceMips[MipIndex];
+			const FImage& Image = (*pSourceMips)[MipIndex];
+
+			// copy mips over + processing
+			// this is a code dupe of the processing done in GenerateMipChain
 
 			// create base for the mip chain
 			FImage* Mip = new(OutMipChain) FImage();
@@ -3045,6 +3150,8 @@ private:
 				// Generate the base mip from the long-lat source image.
 				GenerateBaseCubeMipFromLongitudeLatitude2D(Mip, Image, BuildSettings.MaxTextureResolution, BuildSettings.SourceEncodingOverride);
 	
+				check( CopyCount == 1 );
+
 				// note the break here skips other adjustments:
 				break;
 			}
@@ -3119,17 +3226,29 @@ private:
 		check( GenerateCount == NumOutputMips - OutMipChain.Num() );
 
 		// Generate any missing mips in the chain.
-		if (NumOutputMips > OutMipChain.Num())
+		if ( GenerateCount > 0 )
 		{
 			// Do angular filtering of cubemaps if requested.
 			if (BuildSettings.MipGenSettings == TMGS_Angular)
 			{
+				check( BuildSettings.bCubemap );
+				// note TMGS_Angular forces dim to next lower power of 2
+
+				// note GenerateAngularFilteredMips reprocesses ALL the mips, not just GenerateCount
+				// this should probably be outside the GenerateCount check (eg. always done, even if GenerateCount == 0 )
+				// but putting it inside matches existing behavior
+				// I guess it's moot because you can't set NoMipMips or LeaveExisting if you chose Angular
+
 				GenerateAngularFilteredMips(OutMipChain, NumOutputMips, BuildSettings.DiffuseConvolveMipLevel);
 			}
 			else
 			{
-				// @@!! GenerateMipChain should bring us up to NumOutputMips
+				// GenerateMipChain should bring us up to NumOutputMips
 				//	but it doesn't take NumOutputMips as a param, makes its own decision
+				//  we will check that it chose the same mip count after
+
+				// you could pass GenerateCount as the large arg here
+				//  and it should make the same result
 
 				GenerateMipChain(BuildSettings, OutMipChain.Last(), OutMipChain, MAX_uint32);
 			}
diff --git a/Engine/Source/Developer/TextureCompressor/Public/TextureCompressorModule.h b/Engine/Source/Developer/TextureCompressor/Public/TextureCompressorModule.h
--- a/Engine/Source/Developer/TextureCompressor/Public/TextureCompressorModule.h
+++ b/Engine/Source/Developer/TextureCompressor/Public/TextureCompressorModule.h
@@ -207,6 +207,8 @@ struct FTextureBuildSettings
 	// ETextureEncodeSpeed, either Final or Fast.
 	uint8 RepresentsEncodeSpeedNoSend;
 
+	static constexpr uint32 MaxTextureResolutionDefault = TNumericLimits<uint32>::Max();
+
 	/** Default settings. */
 	FTextureBuildSettings()
 		: bDoScaleMipsForAlphaCoverage(false)
@@ -214,7 +216,7 @@ struct FTextureBuildSettings
 		, MipSharpening(0.0f)
 		, DiffuseConvolveMipLevel(0)
 		, SharpenMipKernelSize(2)
-		, MaxTextureResolution(TNumericLimits<uint32>::Max())
+		, MaxTextureResolution(MaxTextureResolutionDefault)
 		, bHDRSource(false)
 		, MipGenSettings(1 /*TMGS_SimpleAverage*/)
 		, bCubemap(false)
diff --git a/Engine/Source/Developer/TextureFormat/Public/Interfaces/ITextureFormat.h b/Engine/Source/Developer/TextureFormat/Public/Interfaces/ITextureFormat.h
--- a/Engine/Source/Developer/TextureFormat/Public/Interfaces/ITextureFormat.h
+++ b/Engine/Source/Developer/TextureFormat/Public/Interfaces/ITextureFormat.h
@@ -14,12 +14,13 @@ struct FTextureBuildSettings;
 struct FTextureFormatCompressorCaps
 {
 	FTextureFormatCompressorCaps()
-		: MaxTextureDimension(TNumericLimits<uint32>::Max())
+		: MaxTextureDimension_DEPRECATED(TNumericLimits<uint32>::Max())
 		, NumMipsInTail(0)
 		, ExtData(0)
 	{ }
 
-	uint32 MaxTextureDimension;
+	// MaxTextureDimension is never set, remove it
+	uint32 MaxTextureDimension_DEPRECATED;
 	uint32 NumMipsInTail;
 	uint32 ExtData;
 };
diff --git a/Engine/Source/Editor/TextureEditor/Private/Customizations/TextureDetailsCustomization.cpp b/Engine/Source/Editor/TextureEditor/Private/Customizations/TextureDetailsCustomization.cpp
--- a/Engine/Source/Editor/TextureEditor/Private/Customizations/TextureDetailsCustomization.cpp
+++ b/Engine/Source/Editor/TextureEditor/Private/Customizations/TextureDetailsCustomization.cpp
@@ -4,7 +4,7 @@
 #include "Misc/MessageDialog.h"
 #include "Widgets/DeclarativeSyntaxSupport.h"
 #include "Engine/Texture.h"
-#include "Engine/Texture2D.h"
+//#include "Engine/Texture2D.h"
 #include "Editor.h"
 #include "DetailLayoutBuilder.h"
 #include "DetailWidgetRow.h"
@@ -137,25 +137,6 @@ void FTextureDetails::CustomizeDetails(IDetailLayoutBuilder& DetailBuilder)
 	}
 }
 
-bool FTextureDetails::CanEditMaxTextureSize() const
-{
-	if (UTexture2D* Texture2D = Cast<UTexture2D>(TextureBeingCustomized.Get()))
-	{
-		if (!Texture2D->Source.IsPowerOfTwo() && Texture2D->PowerOfTwoMode == ETexturePowerOfTwoSetting::None)
-		{
-			return false;
-		}
-	}
-
-	return true;
-}
-
-void FTextureDetails::CreateMaxTextureSizeMessage() const
-{
-	FMessageDialog::Open(EAppMsgType::Ok,
-		LOCTEXT("CannotEditMaxTextureSize", "Maximum Texture Size cannot be changed for this texture as it is a non power of two size. Change the Power of Two Mode to allow it to be padded to a power of two."));
-}
-
 /** @return The value or unset if properties with multiple values are viewed */
 TOptional<int32> FTextureDetails::OnGetMaxTextureSize() const
 {
@@ -171,11 +152,6 @@ TOptional<int32> FTextureDetails::OnGetMaxTextureSize() const
 
 void FTextureDetails::OnMaxTextureSizeChanged(int32 NewValue)
 {
-	if (!CanEditMaxTextureSize())
-	{
-		return;
-	}
-
 	if (bIsUsingSlider)
 	{
 		int32 OrgValue(0);
@@ -196,15 +172,6 @@ void FTextureDetails::OnMaxTextureSizeChanged(int32 NewValue)
 
 void FTextureDetails::OnMaxTextureSizeCommitted(int32 NewValue, ETextCommit::Type CommitInfo)
 {
-	if (!CanEditMaxTextureSize())
-	{
-		if (CommitInfo == ETextCommit::OnEnter)
-		{
-			CreateMaxTextureSizeMessage();
-		}
-		return;
-	}
-
 	MaxTextureSizePropertyHandle->SetValue(NewValue);
 }
 
@@ -213,11 +180,6 @@ void FTextureDetails::OnMaxTextureSizeCommitted(int32 NewValue, ETextCommit::Typ
  */
 void FTextureDetails::OnBeginSliderMovement()
 {
-	if (!CanEditMaxTextureSize())
-	{
-		return;
-	}
-
 	bIsUsingSlider = true;
 
 	GEditor->BeginTransaction(TEXT("TextureDetails"), LOCTEXT("SetMaximumTextureSize", "Edit Maximum Texture Size"), nullptr /* MaxTextureSizePropertyHandle->GetProperty() */ );
@@ -229,35 +191,11 @@ void FTextureDetails::OnBeginSliderMovement()
  */
 void FTextureDetails::OnEndSliderMovement(int32 NewValue)
 {
-	if (!CanEditMaxTextureSize())
-	{
-		return;
-	}
-
 	bIsUsingSlider = false;
 
 	GEditor->EndTransaction();
 }
 
-bool FTextureDetails::CanEditPowerOfTwoMode(int32 NewPowerOfTwoMode) const
-{
-	if (UTexture2D* Texture2D = Cast<UTexture2D>(TextureBeingCustomized.Get()))
-	{
-		if (!Texture2D->Source.IsPowerOfTwo() && Texture2D->MaxTextureSize > 0 && NewPowerOfTwoMode == ETexturePowerOfTwoSetting::None)
-		{
-			return false;
-		}
-	}
-
-	return true;
-}
-
-void FTextureDetails::CreatePowerOfTwoModeMessage() const
-{
-	FMessageDialog::Open(EAppMsgType::Ok,
-						 LOCTEXT("CannotEditPowerOfTwoMode", "Power of Two Mode cannot be changed to None for this texture as it is a non power of two size and has a Maximum Texture Size override. Change the Maximum Texture Size to 0 before attempting to change the Power of Two Mode."));
-}
-
 void FTextureDetails::OnPropertyResetToDefault() const
 {
 	uint8 CurrentPowerOfTwoMode;
@@ -269,14 +207,6 @@ void FTextureDetails::OnPowerOfTwoModeChanged(TSharedPtr<FString> NewValue, ESel
 {
 	int32 NewPowerOfTwoMode = PowerOfTwoModeComboBoxList.Find(NewValue);
 	check(NewPowerOfTwoMode != INDEX_NONE);
-	if (!CanEditPowerOfTwoMode(NewPowerOfTwoMode))
-	{
-		CreatePowerOfTwoModeMessage();
-		uint8 CurrentPowerOfTwoMode;
-		ensure(PowerOfTwoModePropertyHandle->GetValue(CurrentPowerOfTwoMode) == FPropertyAccess::Success);
-		TextComboBox->SetSelectedItem(PowerOfTwoModeComboBoxList[CurrentPowerOfTwoMode]);
-		return;
-	}
 
 	int32 PowerOfTwoMode = PowerOfTwoModeComboBoxList.Find(NewValue);
 	check(PowerOfTwoMode != INDEX_NONE);
diff --git a/Engine/Source/Editor/TextureEditor/Private/Customizations/TextureDetailsCustomization.h b/Engine/Source/Editor/TextureEditor/Private/Customizations/TextureDetailsCustomization.h
--- a/Engine/Source/Editor/TextureEditor/Private/Customizations/TextureDetailsCustomization.h
+++ b/Engine/Source/Editor/TextureEditor/Private/Customizations/TextureDetailsCustomization.h
@@ -30,11 +30,7 @@ private:
 	void OnMaxTextureSizeCommitted(int32 NewValue, ETextCommit::Type CommitInfo);
 	void OnBeginSliderMovement();
 	void OnEndSliderMovement(int32 NewValue);
-	bool CanEditMaxTextureSize() const;
-	void CreateMaxTextureSizeMessage() const;
 	void OnPowerOfTwoModeChanged(TSharedPtr<FString> NewValue, ESelectInfo::Type SelectInfo);
-	bool CanEditPowerOfTwoMode(int32 NewPowerOfTwoMode) const;
-	void CreatePowerOfTwoModeMessage() const;
 	void OnPropertyResetToDefault() const;
 
 	TSharedPtr<IPropertyHandle> MaxTextureSizePropertyHandle;
diff --git a/Engine/Source/Runtime/Engine/Private/Texture.cpp b/Engine/Source/Runtime/Engine/Private/Texture.cpp
--- a/Engine/Source/Runtime/Engine/Private/Texture.cpp
+++ b/Engine/Source/Runtime/Engine/Private/Texture.cpp
@@ -380,17 +380,15 @@ void UTexture::ValidateSettingsAfterImportOrEdit()
 
 		if ( PowerOfTwoMode != ETexturePowerOfTwoSetting::None )
 		{
+			// @@!! yes it is? fix and remove this?
 			UE_LOG(LogTexture, Warning, TEXT("PowerOfTwoMode not supported for volume Z"));
 			PowerOfTwoMode = ETexturePowerOfTwoSetting::None;
 		}
 	}
-
-	if (!IsPowerOfTwo && (PowerOfTwoMode == ETexturePowerOfTwoSetting::None))
+	
+	if (VirtualTextureStreaming)
 	{
-		// Force NPT textures to have no mipmaps.
-		MipGenSettings = TMGS_NoMipmaps;
-		NeverStream = true;
-		if (VirtualTextureStreaming)
+		if (!IsPowerOfTwo && (PowerOfTwoMode == ETexturePowerOfTwoSetting::None))
 		{
 			UE_LOG(LogTexture, Warning, TEXT("VirtualTextureStreaming not supported for \"%s\", texture size is not a power-of-2"), *GetName());
 			VirtualTextureStreaming = false;
diff --git a/Engine/Source/Runtime/Engine/Private/TextureDerivedData.cpp b/Engine/Source/Runtime/Engine/Private/TextureDerivedData.cpp
--- a/Engine/Source/Runtime/Engine/Private/TextureDerivedData.cpp
+++ b/Engine/Source/Runtime/Engine/Private/TextureDerivedData.cpp
@@ -201,9 +201,19 @@ static void SerializeForKey(FArchive& Ar, const FTextureBuildSettings& Settings)
 		TempByte = Settings.DownscaleOptions; Ar << TempByte;
 	}
 
+	// weird/unncessary way to do this :
+	// should have just serialized the bool bForceAlphaChannel
 	if (Settings.bForceAlphaChannel)
 	{
-		TempGuid = FGuid(0x2C9DF7E3, 0xBC9D413B, 0xBF963C7A, 0x3F27E8B1); // Guid reserved for bForceAlphaChannel feature
+		TempGuid = FGuid(0x2C9DF7E3, 0xBC9D413B, 0xBF963C7A, 0x3F27E8B1);
+		Ar << TempGuid;
+	}
+
+	if ( Settings.MaxTextureResolution != FTextureBuildSettings::MaxTextureResolutionDefault &&
+		( Settings.MipGenSettings == TMGS_LeaveExistingMips || Settings.bDoScaleMipsForAlphaCoverage ) )
+	{
+		// behavior of MaxTextureResolution + LeaveExistingMips or bDoScaleMipsForAlphaCoverage changed, so modify the key :
+		TempGuid = FGuid(0x418B8584, 0x72D54EA5, 0xBA8E8C2B, 0xECC880DE);
 		Ar << TempGuid;
 	}
 
diff --git a/Engine/Source/Runtime/Engine/Private/TextureLODSettings.cpp b/Engine/Source/Runtime/Engine/Private/TextureLODSettings.cpp
--- a/Engine/Source/Runtime/Engine/Private/TextureLODSettings.cpp
+++ b/Engine/Source/Runtime/Engine/Private/TextureLODSettings.cpp
@@ -270,6 +270,7 @@ void UTextureLODSettings::GetMipGenSettings(const UTexture& Texture, TextureMipG
 	}
 
 	// angular filtering only applies to cubemaps
+	// @@!! currently intentionally NOT allowed on cubearrays
 	if (Setting == TMGS_Angular && !Texture.IsA(UTextureCube::StaticClass()))
 	{
 		Setting = TMGS_NoMipmaps;
