Remove NVTT imports from TextureCompressor Module
should only be pulled from TextureFormatDXT

diff --git a/Engine/Source/Developer/TextureCompressor/Private/TextureCompressorModule.cpp b/Engine/Source/Developer/TextureCompressor/Private/TextureCompressorModule.cpp
--- a/Engine/Source/Developer/TextureCompressor/Private/TextureCompressorModule.cpp
+++ b/Engine/Source/Developer/TextureCompressor/Private/TextureCompressorModule.cpp
@@ -1312,7 +1312,19 @@ void ITextureCompressorModule::GenerateMipChain(
 		{
 			const int32 SrcSliceIndex = Settings.bVolume ? (SliceIndex * 2) : SliceIndex;
 			const FImageView2D IntermediateSrcView = FImageView2D::ConstructConst(IntermediateSrc, SrcSliceIndex);
-			const FImageView2D IntermediateSrcView2 = Settings.bVolume ?  FImageView2D::ConstructConst(IntermediateSrc, SrcSliceIndex + 1) : FImageView2D(); // Volume texture mips take 2 slices
+			FImageView2D IntermediateSrcView2;
+			if ( Settings.bVolume )
+			{
+				if ( SrcSliceIndex + 1 < IntermediateSrc.NumSlices )
+				{
+					IntermediateSrcView2 = FImageView2D::ConstructConst(IntermediateSrc, SrcSliceIndex + 1);
+				}
+				else
+				{
+					// nonpow2 volume sizeZ , clamp slice index
+					IntermediateSrcView2 = FImageView2D::ConstructConst(IntermediateSrc, SrcSliceIndex);
+				}
+			}
 			FImageView2D DestView(DestImage, SliceIndex);
 			FImageView2D IntermediateDstView(IntermediateDst, SliceIndex);
 
@@ -2657,9 +2669,6 @@ class FTextureCompressorModule : public ITextureCompressorModule
 {
 public:
 	FTextureCompressorModule()
-#if PLATFORM_WINDOWS
-		:	nvTextureToolsHandle(0)
-#endif	//PLATFORM_WINDOWS
 	{
 	}
 
@@ -2813,36 +2822,13 @@ public:
 	// IModuleInterface implementation.
 	void StartupModule()
 	{
-		// @@!! what's this for ?
-#if PLATFORM_WINDOWS
-	#if PLATFORM_64BITS
-		if (FWindowsPlatformMisc::HasAVX2InstructionSupport())
-		{
-			nvTextureToolsHandle = FPlatformProcess::GetDllHandle(*(FPaths::EngineDir() / TEXT("Binaries/ThirdParty/nvTextureTools/Win64/AVX2/nvtt_64.dll")));
-		}
-		else
-		{
-			nvTextureToolsHandle = FPlatformProcess::GetDllHandle(*(FPaths::EngineDir() / TEXT("Binaries/ThirdParty/nvTextureTools/Win64/nvtt_64.dll")));
-		}
-	#else	//32-bit platform
-		nvTextureToolsHandle = FPlatformProcess::GetDllHandle(*(FPaths::EngineDir() / TEXT("Binaries/ThirdParty/nvTextureTools/Win32/nvtt_.dll")));
-	#endif
-#endif	//PLATFORM_WINDOWS
 	}
 
 	void ShutdownModule()
 	{
-#if PLATFORM_WINDOWS
-		FPlatformProcess::FreeDllHandle(nvTextureToolsHandle);
-		nvTextureToolsHandle = 0;
-#endif
 	}
 
 private:
-#if PLATFORM_WINDOWS
-	// Handle to the nvtt dll
-	void* nvTextureToolsHandle;
-#endif	//PLATFORM_WINDOWS
 
 	bool BuildTextureMips(
 		const TArray<FImage>& InSourceMipChain,
@@ -2853,9 +2839,6 @@ private:
 	{
 		TRACE_CPUPROFILER_EVENT_SCOPE(Texture.BuildTextureMips);
 		
-		// @@!! TEMP
-		//const_cast<FTextureBuildSettings&>(BuildSettings).MaxTextureResolution = 300;
-
 		check(InSourceMipChain.Num() > 0);
 		check(InSourceMipChain[0].SizeX > 0 && InSourceMipChain[0].SizeY > 0 && InSourceMipChain[0].NumSlices > 0);
 
diff --git a/Engine/Source/Developer/TextureCompressor/TextureCompressor.Build.cs b/Engine/Source/Developer/TextureCompressor/TextureCompressor.Build.cs
--- a/Engine/Source/Developer/TextureCompressor/TextureCompressor.Build.cs
+++ b/Engine/Source/Developer/TextureCompressor/TextureCompressor.Build.cs
@@ -20,6 +20,6 @@ public class TextureCompressor : ModuleRules
 			}
 			);
 
-		AddEngineThirdPartyPrivateStaticDependencies(Target, "nvTextureTools");
+		//AddEngineThirdPartyPrivateStaticDependencies(Target, "nvTextureTools");
 	}
 }
diff --git a/Engine/Source/Developer/TextureFormat/TextureFormat.Build.cs b/Engine/Source/Developer/TextureFormat/TextureFormat.Build.cs
--- a/Engine/Source/Developer/TextureFormat/TextureFormat.Build.cs
+++ b/Engine/Source/Developer/TextureFormat/TextureFormat.Build.cs
@@ -13,6 +13,6 @@ public class TextureFormat : ModuleRules
 			}
 			);
 
-		AddEngineThirdPartyPrivateStaticDependencies(Target, "nvTextureTools");
+		//AddEngineThirdPartyPrivateStaticDependencies(Target, "nvTextureTools");
 	}
 }
