use FMath::FloorLog2 instead of writing my own

diff --git a/Engine/Source/Developer/TextureCompressor/Private/TextureCompressorModule.cpp b/Engine/Source/Developer/TextureCompressor/Private/TextureCompressorModule.cpp
--- a/Engine/Source/Developer/TextureCompressor/Private/TextureCompressorModule.cpp
+++ b/Engine/Source/Developer/TextureCompressor/Private/TextureCompressorModule.cpp
@@ -1447,8 +1447,8 @@ struct FImageViewLongLat
 	/** Makes a filtered lookup. */
 	FLinearColor LookupFiltered(float X, float Y) const
 	{
-		int32 X0 = (int32)floor(X);
-		int32 Y0 = (int32)floor(Y);
+		int32 X0 = (int32)floorf(X);
+		int32 Y0 = (int32)floorf(Y);
 
 		float FracX = X - X0;
 		float FracY = Y - Y0;
@@ -2652,16 +2652,6 @@ static void NormalizeMip(FImage& InOutMip)
 	}
 }
 
-static uint32 FloorLogTwo(uint32 X) 
-{
-	uint32 L2 = FMath::CeilLogTwo(X);
-	if ( ! FMath::IsPowerOfTwo(X) )
-	{
-		L2--;
-	}
-	return L2;
-}
-
 /**
  * Texture compression module
  */
@@ -3093,7 +3083,7 @@ private:
 		
 			// NumOutputMips is the number of mips that would be made if you made a full mip chain
 			//  eg. 256 makes 9 mips , 300 also makes 9 mips
-			NumOutputMips = 1 + FloorLogTwo(
+			NumOutputMips = 1 + FMath::FloorLog2(
 				bLongLatCubemap ?
 				ComputeLongLatCubemapExtents(TopMip, BuildSettings.MaxTextureResolution) :
 				FMath::Max3(TopMip.SizeX, TopMip.SizeY, TopMipSizeZ) );
