fix TextureCompressor reading from TArray that could resize

diff --git a/Engine/Source/Developer/TextureCompressor/Private/TextureCompressorModule.cpp b/Engine/Source/Developer/TextureCompressor/Private/TextureCompressorModule.cpp
--- a/Engine/Source/Developer/TextureCompressor/Private/TextureCompressorModule.cpp
+++ b/Engine/Source/Developer/TextureCompressor/Private/TextureCompressorModule.cpp
@@ -3033,6 +3033,10 @@ private:
 					BuildSettings.MaxTextureResolution
 					);
 
+				// make sure BuildSourceImageMips doesn't reallocate :
+				constexpr int BuildSourceImageMipsMaxCount = 16;
+				BuildSourceImageMips.Empty(BuildSourceImageMipsMaxCount);
+
 				// Max Texture Size resizing happens here :
 				// note we do not check for TMGS_Angular here
 				GenerateMipChain(BuildSettings, bSuitableFormat ? BaseImage : Temp, BuildSourceImageMips, 1);
@@ -3042,6 +3046,7 @@ private:
 				{
 					// note: now making mips one by one, rather than N in one call
 					//	this is not exactly the same if AlphaCoverage processing is on
+					check( BuildSourceImageMips.Num() < BuildSourceImageMipsMaxCount );
 					GenerateMipChain(BuildSettings, BuildSourceImageMips.Last(), BuildSourceImageMips, 1);
 				}
 			
