TextureCompressorModule nop comment
every mip is made twice

diff --git a/Engine/Source/Developer/TextureCompressor/Private/TextureCompressorModule.cpp b/Engine/Source/Developer/TextureCompressor/Private/TextureCompressorModule.cpp
--- a/Engine/Source/Developer/TextureCompressor/Private/TextureCompressorModule.cpp
+++ b/Engine/Source/Developer/TextureCompressor/Private/TextureCompressorModule.cpp
@@ -1328,6 +1328,7 @@ void ITextureCompressorModule::GenerateMipChain(
 			FImageView2D DestView(DestImage, SliceIndex);
 			FImageView2D IntermediateDstView(IntermediateDst, SliceIndex);
 
+			// DestView is the output mip
 			GenerateSharpenedMipB8G8R8A8(
 				IntermediateSrcView, 
 				IntermediateSrcView2,
@@ -1342,9 +1343,17 @@ void ITextureCompressorModule::GenerateMipChain(
 				Settings.MipGenSettings == TMGS_Unfiltered);
 
 			// generate IntermediateDstImage:
+			// IntermediateDstImage will be the source for the next mip
 			if ( Settings.bDownsampleWithAverage )
 			{
 				// down sample without sharpening for the next iteration
+				// bDownsampleWithAverage comes from GetMipGenSettings
+				// it is on by default for all cases except Blur
+				// it means every mip is generated *twice*
+				// the output mip is made above using Sharpen from the IntermediateSrc
+				// then the next source is made here using 2x2 ("SimpleAverage")
+				// the next IntermediateSrc is not my outputmip, it's what is made here
+
 				GenerateSharpenedMipB8G8R8A8(
 					IntermediateSrcView,
 					IntermediateSrcView2,
