Adds option to use separable filter for generating mips.
This allows faster texture cooking, but introduces slight changes to newly computed pixel values.
Quality of texture mips stays the same, but values are different due to different order of floating-point operations.
To keep patch sizes small, use old same code as before for generating mips for existing textures - will use new filter only for saved textures.
Currently this setting is disabled for new textures by default. See comment in Texture.cpp how to enable it.

diff --git a/Engine/Source/Developer/TextureBuild/Private/TextureBuildFunction.cpp b/Engine/Source/Developer/TextureBuild/Private/TextureBuildFunction.cpp
--- a/Engine/Source/Developer/TextureBuild/Private/TextureBuildFunction.cpp
+++ b/Engine/Source/Developer/TextureBuild/Private/TextureBuildFunction.cpp
@@ -104,6 +104,7 @@ static FTextureBuildSettings ReadBuildSettingsFromCompactBinary(const FCbObjectV
 	ReadCbField(ColorAdjustmentCbObj["AdjustHue"], ColorAdjustment.AdjustHue);
 	ReadCbField(ColorAdjustmentCbObj["AdjustMinAlpha"], ColorAdjustment.AdjustMinAlpha);
 	ReadCbField(ColorAdjustmentCbObj["AdjustMaxAlpha"], ColorAdjustment.AdjustMaxAlpha);
+	BuildSettings.bUseNewMipFilter = Object["bUseNewMipFilter"].AsBool(BuildSettings.bUseNewMipFilter);
 	BuildSettings.bDoScaleMipsForAlphaCoverage = Object["bDoScaleMipsForAlphaCoverage"].AsBool(BuildSettings.bDoScaleMipsForAlphaCoverage);
 	ReadCbField(Object["AlphaCoverageThresholds"], BuildSettings.AlphaCoverageThresholds);
 	ReadCbField(Object["MipSharpening"], BuildSettings.MipSharpening);
diff --git a/Engine/Source/Developer/TextureCompressor/Private/TextureCompressorModule.cpp b/Engine/Source/Developer/TextureCompressor/Private/TextureCompressorModule.cpp
--- a/Engine/Source/Developer/TextureCompressor/Private/TextureCompressorModule.cpp
+++ b/Engine/Source/Developer/TextureCompressor/Private/TextureCompressorModule.cpp
@@ -113,7 +113,40 @@ static const FLinearColor& LookupSourceMip(const FImageView2D& SourceImageData,
 	{
 		check(0);
 	}
-	return SourceImageData.Access(X,Y);
+	return SourceImageData.Access(X, Y);
+}
+
+// Same functionality as above, but for 1D lookup with explicit size
+template <EMipGenAddressMode AddressMode>
+static const FLinearColor& LookupSourceMip(const FLinearColor* Data, int32 Size, int32 X)
+{
+	if (AddressMode == MGTAM_Wrap)
+	{
+		// wrap
+		// ! requires pow2 sizes
+		checkSlow( FMath::IsPowerOfTwo(Size) );
+
+		X = (int32)((uint32)X) & (Size - 1);
+	}
+	else if (AddressMode == MGTAM_Clamp)
+	{
+		// clamp
+		X = FMath::Clamp(X, 0, Size - 1);
+	}
+	else if (AddressMode == MGTAM_BorderBlack)
+	{
+		// border color 0
+		if ((uint32)X >= (uint32)Size)
+		{
+			static FLinearColor Black(0, 0, 0, 0);
+			return Black;
+		}
+	}
+	else
+	{
+		check(0);
+	}
+	return Data[X];
 }
 
 // Kernel class for image filtering operations like image downsampling
@@ -135,7 +168,7 @@ public:
 			TableSize1D = MaxKernelExtend;
 		}
 
-		float Table1D[MaxKernelExtend];
+		float* Table1D = KernelWeights1D;
 		float NegativeTable1D[MaxKernelExtend];
 
 		FilterTableSize = TableSize1D;
@@ -145,6 +178,7 @@ public:
 			// 2x2 kernel: simple average
 			// SharpenFactor is ignored
 			// this is TMGS_SimpleAverage
+			Table1D[0] = Table1D[1] = 0.5f;
 			KernelWeights[0] = KernelWeights[1] = KernelWeights[2] = KernelWeights[3] = 0.25f;
 			return;
 		}
@@ -207,14 +241,13 @@ public:
 		return FilterTableSize;
 	}
 
-	inline float GetAt(uint32 X, uint32 Y) const
+	inline float Get1D(uint32 X) const
 	{
 		checkSlow(X < FilterTableSize);
-		checkSlow(Y < FilterTableSize);
-		return KernelWeights[X + Y * FilterTableSize];
+		return KernelWeights1D[X];
 	}
 
-	inline float& GetRefAt(uint32 X, uint32 Y)
+	inline float GetAt(uint32 X, uint32 Y) const
 	{
 		checkSlow(X < FilterTableSize);
 		checkSlow(Y < FilterTableSize);
@@ -332,6 +365,7 @@ private:
 	uint32 FilterTableSize;
 	// normalized, means the sum of it should be 1.0f
 	float KernelWeights[MaxKernelExtend * MaxKernelExtend];
+	float KernelWeights1D[MaxKernelExtend];
 };
 
 static float DetermineScaledThreshold(float Threshold, float Scale)
@@ -563,9 +597,9 @@ static void GenerateMip2x2Simple(
 		int32 EndIndex = FMath::Min(StartIndex + NumRowsEachJob, DestImageData.SizeY);
 		for (int32 DestY = StartIndex; DestY < EndIndex; ++DestY)
 		{
-			float* DestRow = (float*)&DestImageData.Access(0, DestY);
-			const float* SourceRow0 = (const float*)&SourceImageData.Access(0, 2*DestY);
-			const float* SourceRow1 = (const float*)&SourceImageData.Access(0, 2*DestY+1);
+			float* DestRow = &DestImageData.Access(0, DestY).Component(0);
+			const float* SourceRow0 = &SourceImageData.Access(0, 2*DestY).Component(0);
+			const float* SourceRow1 = &SourceImageData.Access(0, 2*DestY+1).Component(0);
 
 			const VectorRegister4Float Mul = VectorSetFloat1(0.25f);
 			for (int32 DestX = 0; DestX < DestImageData.SizeX; DestX++)
@@ -586,13 +620,14 @@ static void GenerateMip2x2Simple(
 }
 
 template <EMipGenAddressMode AddressMode>
-static void GenerateMipUnfiltered(const FImageView2D& SourceImageData, FImageView2D& DestImageData, const FVector4f AlphaScale, uint32 ScaleFactor)
+static void GenerateMipUnfiltered(const FImageView2D& SourceImageData, FImageView2D& DestImageData, FVector4f AlphaScale, uint32 ScaleFactor)
 {
 	int32 NumRowsEachJob;
 	int32 NumJobs = ImageParallelForComputeNumJobsForRows(NumRowsEachJob, DestImageData.SizeX, DestImageData.SizeY);
 
 	ParallelFor(TEXT("Texture.GenerateMipUnfiltered.PF"), NumJobs, 1, [&](int32 Index)
 	{
+		VectorRegister4Float AlphaScaleV = VectorLoad(&AlphaScale[0]);
 		int32 StartIndex = Index * NumRowsEachJob;
 		int32 EndIndex = FMath::Min(StartIndex + NumRowsEachJob, DestImageData.SizeY);
 		for (int32 DestY = StartIndex; DestY < EndIndex; ++DestY)
@@ -602,26 +637,27 @@ static void GenerateMipUnfiltered(const FImageView2D& SourceImageData, FImageVie
 				const int32 SourceX = DestX * ScaleFactor;
 				const int32 SourceY = DestY * ScaleFactor;
 				const FLinearColor& Sample = LookupSourceMip<AddressMode>(SourceImageData, SourceX, SourceY);
-				VectorRegister4Float FilteredColor = VectorLoad((const float*)&Sample);
+				VectorRegister4Float FilteredColor = VectorLoad(&Sample.Component(0));
 
 				// Apply computed alpha scales to each channel
-				FilteredColor = VectorMultiply(FilteredColor, VectorLoad((const float*)&AlphaScale));
+				FilteredColor = VectorMultiply(FilteredColor, AlphaScaleV);
 
 				// Set the destination pixel.
-				VectorStore(FilteredColor, (float*)&DestImageData.Access(DestX, DestY));
+				VectorStore(FilteredColor, &DestImageData.Access(DestX, DestY).Component(0));
 			}
 		}
 	});
 }
 
 template <EMipGenAddressMode AddressMode>
-static void GenerateMip2x2(const FImageView2D& SourceImageData, FImageView2D& DestImageData, const FVector4f AlphaScale, uint32 ScaleFactor)
+static void GenerateMip2x2(const FImageView2D& SourceImageData, FImageView2D& DestImageData, FVector4f AlphaScale, uint32 ScaleFactor)
 {
 	int32 NumRowsEachJob;
 	int32 NumJobs = ImageParallelForComputeNumJobsForRows(NumRowsEachJob, DestImageData.SizeX, DestImageData.SizeY);
 
 	ParallelFor(TEXT("Texture.GenerateMip2x2.PF"), NumJobs, 1, [&](int32 Index)
 	{
+		VectorRegister4Float AlphaScaleV = VectorLoad(&AlphaScale[0]);
 		int32 StartIndex = Index * NumRowsEachJob;
 		int32 EndIndex = FMath::Min(StartIndex + NumRowsEachJob, DestImageData.SizeY);
 		for (int32 DestY = StartIndex; DestY < EndIndex; ++DestY)
@@ -631,29 +667,84 @@ static void GenerateMip2x2(const FImageView2D& SourceImageData, FImageView2D& De
 				const int32 SourceX = DestX * ScaleFactor;
 				const int32 SourceY = DestY * ScaleFactor;
 
-				VectorRegister4Float A = VectorLoad((const float*)&LookupSourceMip<AddressMode>(SourceImageData, SourceX + 0, SourceY + 0));
-				VectorRegister4Float B = VectorLoad((const float*)&LookupSourceMip<AddressMode>(SourceImageData, SourceX + 1, SourceY + 0));
-				VectorRegister4Float C = VectorLoad((const float*)&LookupSourceMip<AddressMode>(SourceImageData, SourceX + 0, SourceY + 1));
-				VectorRegister4Float D = VectorLoad((const float*)&LookupSourceMip<AddressMode>(SourceImageData, SourceX + 1, SourceY + 1));
+				VectorRegister4Float A = VectorLoad(&LookupSourceMip<AddressMode>(SourceImageData, SourceX + 0, SourceY + 0).Component(0));
+				VectorRegister4Float B = VectorLoad(&LookupSourceMip<AddressMode>(SourceImageData, SourceX + 1, SourceY + 0).Component(0));
+				VectorRegister4Float C = VectorLoad(&LookupSourceMip<AddressMode>(SourceImageData, SourceX + 0, SourceY + 1).Component(0));
+				VectorRegister4Float D = VectorLoad(&LookupSourceMip<AddressMode>(SourceImageData, SourceX + 1, SourceY + 1).Component(0));
 				VectorRegister4Float FilteredColor = VectorAdd(VectorAdd(VectorAdd(A, B), C), D);
 				FilteredColor = VectorMultiply(FilteredColor, VectorSetFloat1(0.25f));
 
 				// Apply computed alpha scales to each channel
-				FilteredColor = VectorMultiply(FilteredColor, VectorLoad((const float*)&AlphaScale));
+				FilteredColor = VectorMultiply(FilteredColor, AlphaScaleV);
 
 				// Set the destination pixel.
-				VectorStore(FilteredColor, (float*)&DestImageData.Access(DestX, DestY));
+				VectorStore(FilteredColor, &DestImageData.Access(DestX, DestY).Component(0));
 			}
 		}
 	});
 }
 
+static void AllocateTempForMips(
+	TArray<FLinearColor>& TempData,
+	int32 SourceSizeX,
+	int32 SourceSizeY,
+	int32 DestSizeX,
+	int32 DestSizeY,
+	bool bDoScaleMipsForAlphaCoverage,
+	const FImageKernel2D& Kernel,
+	uint32 ScaleFactor,
+	bool bSharpenWithoutColorShift,
+	bool bUnfiltered,
+	bool bUseNewMipFilter)
+{
+	if (!bUseNewMipFilter)
+	{
+		// No need for extra memory if using old 2D filter
+		return;
+	}
+
+	int32 KernelFilterTableSize = (int32)Kernel.GetFilterTableSize();
+
+	if (KernelFilterTableSize == 2 &&
+		ScaleFactor == 2 &&
+		DestSizeX * 2 == SourceSizeX &&
+		DestSizeY * 2 == SourceSizeY &&
+		!bDoScaleMipsForAlphaCoverage &&
+		!bUnfiltered)
+	{
+		// Will use GenerateMip2x2Simple
+		return;
+	}
+
+	if (bUnfiltered)
+	{
+		// Will use GenerateMipUnfiltered
+		return;
+	}
+
+	if (KernelFilterTableSize == 2)
+	{
+		// Will use GenerateMip2x2
+		return;
+	}
+
+	int32 NumRowsEachJob;
+	int32 NumJobs = ImageParallelForComputeNumJobsForRows(NumRowsEachJob, DestSizeX, DestSizeY);
+
+	// Enough bytes to have one row per source width for each job thread
+	// Make sure the rows do not overlap on same cache line
+	int64 SizeInBytes = (sizeof(FLinearColor) * SourceSizeX + PLATFORM_CACHE_LINE_SIZE) * NumJobs;
+
+	int64 ElementCount = FMath::DivideAndRoundUp<int64>(SizeInBytes, sizeof(FLinearColor));
+	TempData.AddUninitialized(ElementCount);
+}
+
 template <EMipGenAddressMode AddressMode, bool bSharpenWithoutColorShift, int KernelSize = 0>
 static void GenerateMipSharpened(
 	const FImageView2D& SourceImageData,
 	FImageView2D& DestImageData,
 	const FImageKernel2D& Kernel,
-	const FVector4f AlphaScale,
+	FVector4f AlphaScale,
 	uint32 ScaleFactor)
 {
 	int32 NumRowsEachJob;
@@ -669,8 +760,11 @@ static void GenerateMipSharpened(
 		// KernelFilterTableSize should be even for standard down-sampling
 		const int32 KernelCenter = (KernelFilterTableSize - 1) / 2;
 
+		VectorRegister4Float AlphaScaleV = VectorLoad(&AlphaScale[0]);
+
 		int32 StartIndex = Index * NumRowsEachJob;
 		int32 EndIndex = FMath::Min(StartIndex + NumRowsEachJob, DestImageData.SizeY);
+
 		for (int32 DestY = StartIndex; DestY < EndIndex; ++DestY)
 		{
 			for (int32 DestX = 0; DestX < DestImageData.SizeX; DestX++)
@@ -685,7 +779,7 @@ static void GenerateMipSharpened(
 					{
 						const FLinearColor& Sample = LookupSourceMip<AddressMode>(SourceImageData, SourceX + KernelX - KernelCenter, SourceY + KernelY - KernelCenter);
 						VectorRegister4Float Weight = VectorSetFloat1(Kernel.GetAt(KernelX, KernelY));
-						VectorRegister4Float WeightSample = VectorMultiply(Weight, VectorLoad((const float*)&Sample));
+						VectorRegister4Float WeightSample = VectorMultiply(Weight, VectorLoad(&Sample.Component(0)));
 						Color = VectorAdd(Color, WeightSample);
 					}
 				}
@@ -700,10 +794,10 @@ static void GenerateMipSharpened(
 					VectorRegister4Float NewLuminance = VectorDot3(SharpenedColor, LuminanceWeights);
 
 					// simple 2x2 kernel to compute the color
-					VectorRegister4Float A = VectorLoad((const float*)&LookupSourceMip<AddressMode>(SourceImageData, SourceX + 0, SourceY + 0));
-					VectorRegister4Float B = VectorLoad((const float*)&LookupSourceMip<AddressMode>(SourceImageData, SourceX + 1, SourceY + 0));
-					VectorRegister4Float C = VectorLoad((const float*)&LookupSourceMip<AddressMode>(SourceImageData, SourceX + 0, SourceY + 1));
-					VectorRegister4Float D = VectorLoad((const float*)&LookupSourceMip<AddressMode>(SourceImageData, SourceX + 1, SourceY + 1));
+					VectorRegister4Float A = VectorLoad(&LookupSourceMip<AddressMode>(SourceImageData, SourceX + 0, SourceY + 0).Component(0));
+					VectorRegister4Float B = VectorLoad(&LookupSourceMip<AddressMode>(SourceImageData, SourceX + 1, SourceY + 0).Component(0));
+					VectorRegister4Float C = VectorLoad(&LookupSourceMip<AddressMode>(SourceImageData, SourceX + 0, SourceY + 1).Component(0));
+					VectorRegister4Float D = VectorLoad(&LookupSourceMip<AddressMode>(SourceImageData, SourceX + 1, SourceY + 1).Component(0));
 					VectorRegister4Float FilteredColor = VectorAdd(VectorAdd(VectorAdd(A, B), C), D);
 					FilteredColor = VectorMultiply(FilteredColor, VectorSetFloat1(0.25f));
 
@@ -719,37 +813,111 @@ static void GenerateMipSharpened(
 					FilteredColor = VectorSelect(AlphaMask, SharpenedColor, FilteredColor);
 
 					// Apply computed alpha scales to each channel
-					FilteredColor = VectorMultiply(FilteredColor, VectorLoad((const float*)&AlphaScale));
+					FilteredColor = VectorMultiply(FilteredColor, AlphaScaleV);
 
 					// Set the destination pixel.
-					VectorStore(FilteredColor, (float*)&DestImageData.Access(DestX, DestY));
+					VectorStore(FilteredColor, &DestImageData.Access(DestX, DestY).Component(0));
 				}
 				else
 				{
 					// Apply computed alpha scales to each channel
-					Color = VectorMultiply(Color, VectorLoad((const float*)&AlphaScale));
+					Color = VectorMultiply(Color, AlphaScaleV);
 
 					// Set the destination pixel.
-					VectorStore(Color, (float*)&DestImageData.Access(DestX, DestY));
+					VectorStore(Color, &DestImageData.Access(DestX, DestY).Component(0));
 				}
+			}
+		}
+	});
+}
 
+template <EMipGenAddressMode AddressMode, int KernelSize = 0>
+static void GenerateMipSharpenedSeparable(
+	const FImageView2D& SourceImageData,
+	TArray<FLinearColor>& TempData,
+	FImageView2D& DestImageData,
+	const FImageKernel2D& Kernel,
+	FVector4f AlphaScale,
+	uint32 ScaleFactor)
+{
+	int32 NumRowsEachJob;
+	int32 NumJobs = ImageParallelForComputeNumJobsForRows(NumRowsEachJob, DestImageData.SizeX, DestImageData.SizeY);
+
+	// Verify that caller has allocated proper amount of bytes for temporary storage
+	// This is allocated in AllocateTempForMips function above
+	check(TempData.Num() >= SourceImageData.SizeX * NumJobs);
+
+	ParallelFor(TEXT("Texture.GenerateMipSharpenedSeparable.PF"), NumJobs, 1, [&](int32 Index)
+	{
+		// In case kernel size is passed as template argument use it as constant
+		// This will allow compiler to unroll inner loops below
+		const int32 KernelFilterTableSize = KernelSize ? KernelSize : (int32)Kernel.GetFilterTableSize();
+
+		// if KernelFilterTableSize is odd, centered in-place filter can be applied
+		// KernelFilterTableSize should be even for standard down-sampling
+		const int32 KernelCenter = (KernelFilterTableSize - 1) / 2;
+
+		VectorRegister4Float AlphaScaleV = VectorLoad(&AlphaScale[0]);
+
+		int32 StartIndex = Index * NumRowsEachJob;
+		int32 EndIndex = FMath::Min(StartIndex + NumRowsEachJob, DestImageData.SizeY);
+
+		FLinearColor* Temp = &TempData[Index * (TempData.Num() / NumJobs)];
+		for (int32 DestY = StartIndex; DestY < EndIndex; ++DestY)
+		{
+			const int32 SourceY = DestY * ScaleFactor;
+
+			for (int32 SourceX = 0; SourceX < SourceImageData.SizeX; SourceX++)
+			{
+				VectorRegister4Float Color = VectorZeroFloat();
+				for (int32 KernelY = 0; KernelY < KernelFilterTableSize; ++KernelY)
+				{
+					const FLinearColor& Sample = LookupSourceMip<AddressMode>(SourceImageData, SourceX, SourceY + KernelY - KernelCenter);
+					VectorRegister4Float Weight = VectorSetFloat1(Kernel.Get1D(KernelY));
+					VectorRegister4Float WeightSample = VectorMultiply(Weight, VectorLoad(&Sample.Component(0)));
+					Color = VectorAdd(Color, WeightSample);
+				}
+				VectorStore(Color, &Temp[SourceX].Component(0));
+			}
+
+			for (int32 DestX = 0; DestX < DestImageData.SizeX; DestX++)
+			{
+				const int32 SourceX = DestX * ScaleFactor;
+
+				VectorRegister4Float Color = VectorZeroFloat();
+				for (int32 KernelX = 0; KernelX < KernelFilterTableSize; ++KernelX)
+				{
+					const FLinearColor& Sample = LookupSourceMip<AddressMode>(Temp, SourceImageData.SizeX, SourceX + KernelX - KernelCenter);
+					VectorRegister4Float Weight = VectorSetFloat1(Kernel.Get1D(KernelX));
+					VectorRegister4Float WeightSample = VectorMultiply(Weight, VectorLoad(&Sample.Component(0)));
+					Color = VectorAdd(Color, WeightSample);
+				}
+
+				// Apply computed alpha scales to each channel
+				Color = VectorMultiply(Color, AlphaScaleV);
+
+				// Set the destination pixel.
+				VectorStore(Color, &DestImageData.Access(DestX, DestY).Component(0));
 			}
 		}
 	});
 }
 
 /**
-* Generates a mip-map for an 2D B8G8R8A8 image using a 4x4 filter with sharpening
+* Generates a mip-map for an 2D B8G8R8A8 image using a filter with possible sharpening
 * @param SourceImageData - The source image's data.
+* @param TempData - Temporary storage required by separable mip filter, call AllocateTempForMips function to allocate it
 * @param DestImageData - The destination image's data.
 * @param ImageFormat - The format of both the source and destination images.
 * @param FilterTable2D - [FilterTableSize * FilterTableSize]
 * @param FilterTableSize - >= 2
 * @param ScaleFactor 1 / 2:for downsampling
+* @param bUseNewMipFilter - pass true to use new separatble mip filter
 */
 template <EMipGenAddressMode AddressMode>
 static void GenerateSharpenedMipB8G8R8A8Templ(
 	const FImageView2D& SourceImageData, 
+	TArray<FLinearColor>& TempData,
 	FImageView2D& DestImageData, 
 	bool bDoScaleMipsForAlphaCoverage,
 	const FVector4f AlphaCoverages,
@@ -757,7 +925,8 @@ static void GenerateSharpenedMipB8G8R8A8Templ(
 	const FImageKernel2D& Kernel,
 	uint32 ScaleFactor,
 	bool bSharpenWithoutColorShift,
-	bool bUnfiltered)
+	bool bUnfiltered,
+	bool bUseNewMipFilter)
 {
 	check( ScaleFactor == 1 || ScaleFactor == 2 );
 	check( (SourceImageData.SizeX/ScaleFactor) == DestImageData.SizeX || DestImageData.SizeX == 1 );
@@ -771,7 +940,8 @@ static void GenerateSharpenedMipB8G8R8A8Templ(
 		// 2x2 is always box filter
 		check( Kernel.GetAt(0,0) == 0.25f );
 	}
-	
+
+	// Keep conditions here in sync with same conditions inside AllocateTempForMips function
 	if ( KernelFilterTableSize == 2 &&
 		ScaleFactor == 2 &&
 		DestImageData.SizeX*2 == SourceImageData.SizeX &&
@@ -802,26 +972,40 @@ static void GenerateSharpenedMipB8G8R8A8Templ(
 		return;
 	}
 
-	if (bSharpenWithoutColorShift)
+	if (bUseNewMipFilter)
 	{
 		if (KernelFilterTableSize == 8)
 		{
-			GenerateMipSharpened<AddressMode, true, 8>(SourceImageData, DestImageData, Kernel, AlphaScale, ScaleFactor);
+			GenerateMipSharpenedSeparable<AddressMode, 8>(SourceImageData, TempData, DestImageData, Kernel, AlphaScale, ScaleFactor);
 		}
 		else
 		{
-			GenerateMipSharpened<AddressMode, true>(SourceImageData, DestImageData, Kernel, AlphaScale, ScaleFactor);
+			GenerateMipSharpenedSeparable<AddressMode>(SourceImageData, TempData, DestImageData, Kernel, AlphaScale, ScaleFactor);
 		}
 	}
 	else
 	{
-		if (KernelFilterTableSize == 8)
+		if (bSharpenWithoutColorShift)
 		{
-			GenerateMipSharpened<AddressMode, false, 8>(SourceImageData, DestImageData, Kernel, AlphaScale, ScaleFactor);
+			if (KernelFilterTableSize == 8)
+			{
+				GenerateMipSharpened<AddressMode, true, 8>(SourceImageData, DestImageData, Kernel, AlphaScale, ScaleFactor);
+			}
+			else
+			{
+				GenerateMipSharpened<AddressMode, true>(SourceImageData, DestImageData, Kernel, AlphaScale, ScaleFactor);
+			}
 		}
 		else
 		{
-			GenerateMipSharpened<AddressMode, false>(SourceImageData, DestImageData, Kernel, AlphaScale, ScaleFactor);
+			if (KernelFilterTableSize == 8)
+			{
+				GenerateMipSharpened<AddressMode, false, 8>(SourceImageData, DestImageData, Kernel, AlphaScale, ScaleFactor);
+			}
+			else
+			{
+				GenerateMipSharpened<AddressMode, false>(SourceImageData, DestImageData, Kernel, AlphaScale, ScaleFactor);
+			}
 		}
 	}
 }
@@ -831,6 +1015,7 @@ static void GenerateSharpenedMipB8G8R8A8Templ(
 static void GenerateSharpenedMipB8G8R8A8(
 	const FImageView2D& SourceImageData, 
 	const FImageView2D& SourceImageData2, // Only used with volume texture.
+	TArray<FLinearColor>& TempData,
 	FImageView2D& DestImageData, 
 	EMipGenAddressMode AddressMode, 
 	bool bDoScaleMipsForAlphaCoverage,
@@ -839,7 +1024,8 @@ static void GenerateSharpenedMipB8G8R8A8(
 	const FImageKernel2D &Kernel,
 	uint32 ScaleFactor,
 	bool bSharpenWithoutColorShift,
-	bool bUnfiltered
+	bool bUnfiltered,
+	bool bUseNewMipFilter
 	)
 {
 	TRACE_CPUPROFILER_EVENT_SCOPE(Texture.GenerateSharpenedMip);
@@ -847,13 +1033,13 @@ static void GenerateSharpenedMipB8G8R8A8(
 	switch(AddressMode)
 	{
 	case MGTAM_Wrap:
-		GenerateSharpenedMipB8G8R8A8Templ<MGTAM_Wrap>(SourceImageData, DestImageData, bDoScaleMipsForAlphaCoverage, AlphaCoverages, AlphaThresholds, Kernel, ScaleFactor, bSharpenWithoutColorShift, bUnfiltered);
+		GenerateSharpenedMipB8G8R8A8Templ<MGTAM_Wrap>(SourceImageData, TempData, DestImageData, bDoScaleMipsForAlphaCoverage, AlphaCoverages, AlphaThresholds, Kernel, ScaleFactor, bSharpenWithoutColorShift, bUnfiltered, bUseNewMipFilter);
 		break;
 	case MGTAM_Clamp:
-		GenerateSharpenedMipB8G8R8A8Templ<MGTAM_Clamp>(SourceImageData, DestImageData, bDoScaleMipsForAlphaCoverage, AlphaCoverages, AlphaThresholds, Kernel, ScaleFactor, bSharpenWithoutColorShift, bUnfiltered);
+		GenerateSharpenedMipB8G8R8A8Templ<MGTAM_Clamp>(SourceImageData, TempData, DestImageData, bDoScaleMipsForAlphaCoverage, AlphaCoverages, AlphaThresholds, Kernel, ScaleFactor, bSharpenWithoutColorShift, bUnfiltered, bUseNewMipFilter);
 		break;
 	case MGTAM_BorderBlack:
-		GenerateSharpenedMipB8G8R8A8Templ<MGTAM_BorderBlack>(SourceImageData, DestImageData, bDoScaleMipsForAlphaCoverage, AlphaCoverages, AlphaThresholds, Kernel, ScaleFactor, bSharpenWithoutColorShift, bUnfiltered);
+		GenerateSharpenedMipB8G8R8A8Templ<MGTAM_BorderBlack>(SourceImageData, TempData, DestImageData, bDoScaleMipsForAlphaCoverage, AlphaCoverages, AlphaThresholds, Kernel, ScaleFactor, bSharpenWithoutColorShift, bUnfiltered, bUseNewMipFilter);
 		break;
 	default:
 		check(0);
@@ -868,13 +1054,13 @@ static void GenerateSharpenedMipB8G8R8A8(
 		switch(AddressMode)
 		{
 		case MGTAM_Wrap:
-			GenerateSharpenedMipB8G8R8A8Templ<MGTAM_Wrap>(SourceImageData2, TempImageData, bDoScaleMipsForAlphaCoverage, AlphaCoverages, AlphaThresholds, Kernel, ScaleFactor, bSharpenWithoutColorShift, bUnfiltered);
+			GenerateSharpenedMipB8G8R8A8Templ<MGTAM_Wrap>(SourceImageData2, TempData, TempImageData, bDoScaleMipsForAlphaCoverage, AlphaCoverages, AlphaThresholds, Kernel, ScaleFactor, bSharpenWithoutColorShift, bUnfiltered, bUseNewMipFilter);
 			break;
 		case MGTAM_Clamp:
-			GenerateSharpenedMipB8G8R8A8Templ<MGTAM_Clamp>(SourceImageData2, TempImageData, bDoScaleMipsForAlphaCoverage, AlphaCoverages, AlphaThresholds, Kernel, ScaleFactor, bSharpenWithoutColorShift, bUnfiltered);
+			GenerateSharpenedMipB8G8R8A8Templ<MGTAM_Clamp>(SourceImageData2, TempData, TempImageData, bDoScaleMipsForAlphaCoverage, AlphaCoverages, AlphaThresholds, Kernel, ScaleFactor, bSharpenWithoutColorShift, bUnfiltered, bUseNewMipFilter);
 			break;
 		case MGTAM_BorderBlack:
-			GenerateSharpenedMipB8G8R8A8Templ<MGTAM_BorderBlack>(SourceImageData2, TempImageData, bDoScaleMipsForAlphaCoverage, AlphaCoverages, AlphaThresholds, Kernel, ScaleFactor, bSharpenWithoutColorShift, bUnfiltered);
+			GenerateSharpenedMipB8G8R8A8Templ<MGTAM_BorderBlack>(SourceImageData2, TempData, TempImageData, bDoScaleMipsForAlphaCoverage, AlphaCoverages, AlphaThresholds, Kernel, ScaleFactor, bSharpenWithoutColorShift, bUnfiltered, bUseNewMipFilter);
 			break;
 		default:
 			check(0);
@@ -1018,6 +1204,13 @@ static void GenerateTopMip(const FImage& SrcImage, FImage& DestImage, const FTex
 
 	DestImage.Init(SrcImage.SizeX, SrcImage.SizeY, SrcImage.NumSlices, SrcImage.Format, SrcImage.GammaSpace);
 
+	const bool bSharpenWithoutColorShift = Settings.bSharpenWithoutColorShift;
+	const bool bUnfiltered = Settings.MipGenSettings == TMGS_Unfiltered;
+	const bool bUseNewMipFilter = Settings.bUseNewMipFilter;
+
+	TArray<FLinearColor> TempData;
+	AllocateTempForMips(TempData, SrcImage.SizeX, SrcImage.SizeY, DestImage.SizeX, DestImage.SizeY, false, KernelDownsample, 1, bSharpenWithoutColorShift, bUnfiltered, bUseNewMipFilter);
+
 	for (int32 SliceIndex = 0; SliceIndex < SrcImage.NumSlices; ++SliceIndex)
 	{
 		FImageView2D SrcView((FImage&)SrcImage, SliceIndex);
@@ -1027,6 +1220,7 @@ static void GenerateTopMip(const FImage& SrcImage, FImage& DestImage, const FTex
 		GenerateSharpenedMipB8G8R8A8(
 			SrcView, 
 			FImageView2D(),
+			TempData,
 			DestView,
 			AddressMode,
 			false,
@@ -1035,7 +1229,8 @@ static void GenerateTopMip(const FImage& SrcImage, FImage& DestImage, const FTex
 			KernelDownsample,
 			1,
 			Settings.bSharpenWithoutColorShift,
-			Settings.MipGenSettings == TMGS_Unfiltered);
+			bUnfiltered,
+			bUseNewMipFilter);
 	}
 }
 
@@ -1065,6 +1260,7 @@ struct FTextureDownscaleSettings
 	int32 BlockSize;
 	float Downscale;
 	uint8 DownscaleOptions;
+	bool UseNewMipFilter;
 };
 
 static void DownscaleImage(const FImage& SrcImage, FImage& DstImage, const FTextureDownscaleSettings& Settings)
@@ -1101,11 +1297,16 @@ static void DownscaleImage(const FImage& SrcImage, FImage& DstImage, const FText
 	FImage Image0;
 	FImage Image1;
 	FImage* ImageChain[2] = {&const_cast<FImage&>(SrcImage), &Image1};
-	bool bUnfiltered = Settings.DownscaleOptions == (uint8)ETextureDownscaleOptions::Unfiltered;
+	const bool bUnfiltered = Settings.DownscaleOptions == (uint8)ETextureDownscaleOptions::Unfiltered;
+	const bool bUseNewMipFilter = Settings.UseNewMipFilter;
 	
 	// Scaledown using 2x2 average, use user specified filtering only for last iteration
 	FImageKernel2D AvgKernel;
 	AvgKernel.BuildSeparatableGaussWithSharpen(2);
+
+	TArray<FLinearColor> TempData;
+	AllocateTempForMips(TempData, SrcImage.SizeX, SrcImage.SizeY, SrcImage.SizeX / 2, SrcImage.SizeY / 2, false, AvgKernel, 2, false, bUnfiltered, bUseNewMipFilter);
+
 	int32 NumIterations = 0;
 	while(Downscale > 2.0f)
 	{
@@ -1117,6 +1318,7 @@ static void DownscaleImage(const FImage& SrcImage, FImage& DstImage, const FText
 		FImageView2D DstImageData(*ImageChain[1], 0);
 		GenerateSharpenedMipB8G8R8A8Templ<MGTAM_Clamp>(
 			SrcImageData, 
+			TempData,
 			DstImageData, 
 			false,
 			FVector4f(0, 0, 0, 0),
@@ -1124,7 +1326,8 @@ static void DownscaleImage(const FImage& SrcImage, FImage& DstImage, const FText
 			AvgKernel, 
 			2, 
 			false,
-			bUnfiltered);
+			bUnfiltered,
+			bUseNewMipFilter);
 
 		if (NumIterations == 0)
 		{
@@ -1288,6 +1491,15 @@ void ITextureCompressorModule::GenerateMipChain(
 		AlphaCoverages = ComputeAlphaCoverage(Settings.AlphaCoverageThresholds, AlphaScales, IntermediateSrcView);
 	}
 
+	TArray<FLinearColor> DownsampleTempData;
+	TArray<FLinearColor> AverageTempData;
+	const bool bDoScaleMipsForAlphaCoverage = Settings.bDoScaleMipsForAlphaCoverage;
+	const bool bSharpenWithoutColorShift = Settings.bSharpenWithoutColorShift;
+	const bool bUnfiltered = Settings.MipGenSettings == TMGS_Unfiltered;
+	const bool bUseNewMipFilter = Settings.bUseNewMipFilter;
+	AllocateTempForMips(DownsampleTempData, BaseMip.SizeX, BaseMip.SizeY, SecondTempImage.SizeX, SecondTempImage.SizeY, bDoScaleMipsForAlphaCoverage, KernelDownsample, 2, bSharpenWithoutColorShift, bUnfiltered, bUseNewMipFilter);
+	AllocateTempForMips(AverageTempData, BaseMip.SizeX, BaseMip.SizeY, SecondTempImage.SizeX, SecondTempImage.SizeY, bDoScaleMipsForAlphaCoverage, KernelSimpleAverage, 2, bSharpenWithoutColorShift, bUnfiltered, bUseNewMipFilter);
+
 	// Generate mips
 	//  default value of MipChainDepth is MAX_uint32, means generate all mips down to 1x1
 	//	(break inside the loop)
@@ -1332,15 +1544,17 @@ void ITextureCompressorModule::GenerateMipChain(
 			GenerateSharpenedMipB8G8R8A8(
 				IntermediateSrcView, 
 				IntermediateSrcView2,
+				DownsampleTempData,
 				DestView,
 				AddressMode,
-				Settings.bDoScaleMipsForAlphaCoverage,
+				bDoScaleMipsForAlphaCoverage,
 				AlphaCoverages,
 				Settings.AlphaCoverageThresholds,
 				KernelDownsample,
 				2,
-				Settings.bSharpenWithoutColorShift,
-				Settings.MipGenSettings == TMGS_Unfiltered);
+				bSharpenWithoutColorShift,
+				bUnfiltered,
+				bUseNewMipFilter);
 
 			// generate IntermediateDstImage:
 			// IntermediateDstImage will be the source for the next mip
@@ -1357,15 +1571,17 @@ void ITextureCompressorModule::GenerateMipChain(
 				GenerateSharpenedMipB8G8R8A8(
 					IntermediateSrcView,
 					IntermediateSrcView2,
+					AverageTempData,
 					IntermediateDstView,
 					AddressMode,
-					Settings.bDoScaleMipsForAlphaCoverage,
+					bDoScaleMipsForAlphaCoverage,
 					AlphaCoverages,
 					Settings.AlphaCoverageThresholds,
 					KernelSimpleAverage,
 					2,
-					Settings.bSharpenWithoutColorShift,
-					Settings.MipGenSettings == TMGS_Unfiltered);
+					bSharpenWithoutColorShift,
+					bUnfiltered,
+					bUseNewMipFilter);
 			}
 		}
 
@@ -2763,6 +2979,9 @@ public:
 			// important to make accurate computation with normal length
 			DefaultSettings.bRenormalizeTopMip = true;
 
+			// use new mip filter setting from build settings
+			DefaultSettings.bUseNewMipFilter = BuildSettings.bUseNewMipFilter;
+
 			// @todo Oodle : filtering the normal map then computing roughness is fundamentally wrong
 			//  we should instead compute the roughness scalar first on the original normap map
 			//  then filter on the roughness scalar
@@ -3181,6 +3400,7 @@ private:
 				DownscaleSettings.Downscale = BuildSettings.Downscale;
 				DownscaleSettings.DownscaleOptions = BuildSettings.DownscaleOptions;
 				DownscaleSettings.BlockSize = 4;
+				DownscaleSettings.UseNewMipFilter = BuildSettings.bUseNewMipFilter;
 		
 				DownscaleImage(*Mip, *Mip, DownscaleSettings);
 			}
diff --git a/Engine/Source/Developer/TextureCompressor/Public/TextureCompressorModule.h b/Engine/Source/Developer/TextureCompressor/Public/TextureCompressorModule.h
--- a/Engine/Source/Developer/TextureCompressor/Public/TextureCompressorModule.h
+++ b/Engine/Source/Developer/TextureCompressor/Public/TextureCompressorModule.h
@@ -77,6 +77,8 @@ struct FTextureBuildSettings
 	bool bDoScaleMipsForAlphaCoverage;
 	/** Channel values to compare to when preserving alpha coverage. */
 	FVector4f AlphaCoverageThresholds;
+	/** Use newer & faster mip generation filter */
+	bool bUseNewMipFilter;
 	/** The desired amount of mip sharpening. */
 	float MipSharpening;
 	/** For angular filtered cubemaps, the mip level which contains convolution with the diffuse cosine lobe. */
@@ -214,6 +216,7 @@ struct FTextureBuildSettings
 	FTextureBuildSettings()
 		: bDoScaleMipsForAlphaCoverage(false)
 		, AlphaCoverageThresholds(0, 0, 0, 0)
+		, bUseNewMipFilter(false)
 		, MipSharpening(0.0f)
 		, DiffuseConvolveMipLevel(0)
 		, SharpenMipKernelSize(2)
diff --git a/Engine/Source/Editor/UnrealEd/Classes/Factories/TextureFactory.h b/Engine/Source/Editor/UnrealEd/Classes/Factories/TextureFactory.h
--- a/Engine/Source/Editor/UnrealEd/Classes/Factories/TextureFactory.h
+++ b/Engine/Source/Editor/UnrealEd/Classes/Factories/TextureFactory.h
@@ -116,6 +116,10 @@ class UNREALED_API UTextureFactory : public UFactory, public IImportSettingsPars
 	UPROPERTY(EditAnywhere, Category=PreserveAlphaCoverage, meta=(ToolTip="Whether mip RGBA should be scaled to preserve the number of pixels with Value >= AlphaCoverageThresholds"))
 	bool bDoScaleMipsForAlphaCoverage = false;
 
+	/** Whether to use newer & faster mip generation filter, same quality but produces slightly different results from previous implementation */
+	UPROPERTY(EditAnywhere, Category=TextureFactory, meta=(ToolTip="Whether to use newer & faster mip generation filter"))
+	bool bUseNewMipFilter = false;
+
 	/** Channel values to compare to when preserving alpha coverage from a mask. */
 	UPROPERTY(EditAnywhere, Category=PreserveAlphaCoverage, meta=(ToolTip="Channel values to compare to when preserving alpha coverage from a mask for mips"))
 	FVector4 AlphaCoverageThresholds = FVector4(0,0,0,0);
diff --git a/Engine/Source/Editor/UnrealEd/Private/Factories/EditorFactories.cpp b/Engine/Source/Editor/UnrealEd/Private/Factories/EditorFactories.cpp
--- a/Engine/Source/Editor/UnrealEd/Private/Factories/EditorFactories.cpp
+++ b/Engine/Source/Editor/UnrealEd/Private/Factories/EditorFactories.cpp
@@ -4034,6 +4034,7 @@ UObject* UTextureFactory::FactoryCreateBinary
 	float								ExistingAdjustMaxAlpha = 1.0f;
 	bool								ExistingbDoScaleMipsForAlphaCoverage = false;
 	FVector4							ExistingAlphaCoverageThresholds = FVector4(0, 0, 0, 0);
+	bool								ExistingbUseNewMipFilter = false;
 	TextureMipGenSettings				ExistingMipGenSettings = TextureMipGenSettings(0);
 	bool								ExistingVirtualTextureStreaming = false;
 
@@ -4123,6 +4124,7 @@ UObject* UTextureFactory::FactoryCreateBinary
 		ExistingFlipGreenChannel = ExistingTexture->bFlipGreenChannel;
 		ExistingbDoScaleMipsForAlphaCoverage = ExistingTexture->bDoScaleMipsForAlphaCoverage;
 		ExistingAlphaCoverageThresholds = ExistingTexture->AlphaCoverageThresholds;
+		ExistingbUseNewMipFilter = ExistingTexture->bUseNewMipFilter;
 		ExistingAdjustBrightness = ExistingTexture->AdjustBrightness;
 		ExistingAdjustBrightnessCurve = ExistingTexture->AdjustBrightnessCurve;
 		ExistingAdjustVibrance = ExistingTexture->AdjustVibrance;
@@ -4240,12 +4242,13 @@ UObject* UTextureFactory::FactoryCreateBinary
 		LODGroup = TEXTUREGROUP_World;
 	}
 
-	Texture->CompressionNone		= NoCompression;
-	Texture->CompressionNoAlpha		= NoAlpha;
-	Texture->DeferCompression		= bDeferCompression;
-	Texture->bDoScaleMipsForAlphaCoverage = bDoScaleMipsForAlphaCoverage;
-	Texture->AlphaCoverageThresholds = AlphaCoverageThresholds;
-	
+	Texture->CompressionNone				= NoCompression;
+	Texture->CompressionNoAlpha				= NoAlpha;
+	Texture->DeferCompression				= bDeferCompression;
+	Texture->bDoScaleMipsForAlphaCoverage	= bDoScaleMipsForAlphaCoverage;
+	Texture->AlphaCoverageThresholds		= AlphaCoverageThresholds;
+	Texture->bUseNewMipFilter				= bUseNewMipFilter;
+
 	if(Texture->MipGenSettings == TMGS_FromTextureGroup)
 	{
 		// unless the loader suggest a different setting
@@ -4307,6 +4310,7 @@ UObject* UTextureFactory::FactoryCreateBinary
 		Texture->DeferCompression = ExistingDeferCompression;
 		Texture->bDoScaleMipsForAlphaCoverage = ExistingbDoScaleMipsForAlphaCoverage;
 		Texture->AlphaCoverageThresholds = ExistingAlphaCoverageThresholds;
+		Texture->bUseNewMipFilter = ExistingbUseNewMipFilter;
 		Texture->bFlipGreenChannel = ExistingFlipGreenChannel;
 		Texture->AdjustBrightness = ExistingAdjustBrightness;
 		Texture->AdjustBrightnessCurve = ExistingAdjustBrightnessCurve;
diff --git a/Engine/Source/Runtime/Engine/Classes/Engine/Texture.h b/Engine/Source/Runtime/Engine/Classes/Engine/Texture.h
--- a/Engine/Source/Runtime/Engine/Classes/Engine/Texture.h
+++ b/Engine/Source/Runtime/Engine/Classes/Engine/Texture.h
@@ -1189,6 +1189,10 @@ public:
 	UPROPERTY(EditAnywhere, BlueprintReadWrite, Category = Texture, meta=(ClampMin = "0", ClampMax = "1.0", EditCondition="bDoScaleMipsForAlphaCoverage"), AdvancedDisplay)
 	FVector4 AlphaCoverageThresholds = FVector4(0,0,0,0);
 
+	/** Whether to use newer & faster mip generation filter, same quality but produces slightly different results from previous implementation */
+	UPROPERTY(EditAnywhere, BlueprintReadWrite, Category=Texture, AdvancedDisplay)
+	bool bUseNewMipFilter = false;
+
 	/** When true the texture's border will be preserved during mipmap generation. */
 	UPROPERTY(EditAnywhere, BlueprintReadWrite, Category=LevelOfDetail, AdvancedDisplay)
 	uint32 bPreserveBorder:1;
diff --git a/Engine/Source/Runtime/Engine/Private/Texture.cpp b/Engine/Source/Runtime/Engine/Private/Texture.cpp
--- a/Engine/Source/Runtime/Engine/Private/Texture.cpp
+++ b/Engine/Source/Runtime/Engine/Private/Texture.cpp
@@ -137,6 +137,7 @@ UTexture::UTexture(const FObjectInitializer& ObjectInitializer)
 	bCustomPropertiesImported = false;
 	bDoScaleMipsForAlphaCoverage = false;
 	AlphaCoverageThresholds = FVector4(0, 0, 0, 0);
+	bUseNewMipFilter = false;
 	PaddingColor = FColor::Black;
 	ChromaKeyColor = FColorList::Magenta;
 	ChromaKeyThreshold = 1.0f / 255.0f;
@@ -632,6 +633,21 @@ void UTexture::Serialize(FArchive& Ar)
 		LoadedMainStreamObjectVersion = Ar.CustomVer(FUE5MainStreamObjectVersion::GUID);
 	}
 
+	// When new mip filter is ready to be enabled:
+	// * change "bUseNewMipFilter = false;" to true in UTexture constructor above
+	// * change "bool bUseNewMipFilter = false;" member to true in "UTexture" class, in "Texture.h" file
+	// * change "bool bUseNewMipFilter = false;" member to true in "UTextureFactory" class, in "TextureFactory.h" file
+	// * change "bUseNewMipFilter(false) to true in "FTextureBuildSettings" constructor, in "TextureCompressorModule.h" file
+	// * change "ExistingbUseNewMipFilter = false" to true in UTextureFactory::FactoryCreateBinary method, in "EditorFactories.cpp" file
+	// * add "TextureUseNewMipFilter" value in "FUE5MainStreamObjectVersion" enum, in "UE5MainStreamObjectVersion.h" file
+	// * uncomment if statement below
+
+	//if (Ar.IsLoading() && Ar.CustomVer(FUE5MainStreamObjectVersion::GUID) < FUE5MainStreamObjectVersion::TextureUseNewMipFilter)
+	//{
+	//	// Old textures should not use new mip filter for maintaining exacty the same output as before (to not increase patch size)
+	//	bUseNewMipFilter = false;
+	//}
+
 	if (Ar.IsLoading() && Ar.CustomVer(FUE5MainStreamObjectVersion::GUID) < FUE5MainStreamObjectVersion::TextureDoScaleMipsForAlphaCoverage)
 	{
 		// bDoScaleMipsForAlphaCoverage was not transmitted in old versions
diff --git a/Engine/Source/Runtime/Engine/Private/TextureDerivedData.cpp b/Engine/Source/Runtime/Engine/Private/TextureDerivedData.cpp
--- a/Engine/Source/Runtime/Engine/Private/TextureDerivedData.cpp
+++ b/Engine/Source/Runtime/Engine/Private/TextureDerivedData.cpp
@@ -231,6 +231,13 @@ static void SerializeForKey(FArchive& Ar, const FTextureBuildSettings& Settings)
 		Ar << TempGuid;
 	}
 
+	// do not change key if old mip filter is used for old textures
+	if (Settings.bUseNewMipFilter)
+	{
+		TempGuid = FGuid(0x27B79A99, 0xE1A5458E, 0xAB619475, 0xCD01AD2A);
+		Ar << TempGuid;
+	}
+
 	// Note - compression quality is added to the DDC by the formats (based on whether they
 	// use them or not).
 	// This is true for:
@@ -736,6 +743,7 @@ static void GetTextureBuildSettings(
 		OutBuildSettings.AlphaCoverageThresholds = FVector4f(0,0,0,0);
 	}
 
+	OutBuildSettings.bUseNewMipFilter = Texture.bUseNewMipFilter;
 	OutBuildSettings.bComputeBokehAlpha = (Texture.LODGroup == TEXTUREGROUP_Bokeh);
 	OutBuildSettings.bReplicateAlpha = false;
 	OutBuildSettings.bReplicateRed = false;
diff --git a/Engine/Source/Runtime/Engine/Private/TextureDerivedDataBuildUtils.cpp b/Engine/Source/Runtime/Engine/Private/TextureDerivedDataBuildUtils.cpp
--- a/Engine/Source/Runtime/Engine/Private/TextureDerivedDataBuildUtils.cpp
+++ b/Engine/Source/Runtime/Engine/Private/TextureDerivedDataBuildUtils.cpp
@@ -120,6 +120,7 @@ static void WriteBuildSettings(FCbWriter& Writer, const FTextureBuildSettings& B
 		// AlphaCoverageThresholds do not affect build if bDoScaleMipsForAlphaCoverage is off
 		WriteCbFieldWithDefault(Writer, "AlphaCoverageThresholds", BuildSettings.AlphaCoverageThresholds, DefaultSettings.AlphaCoverageThresholds);
 	}
+	WriteCbFieldWithDefault<bool>(Writer, "bUseNewMipFilter", BuildSettings.bUseNewMipFilter, DefaultSettings.bUseNewMipFilter);
 	WriteCbFieldWithDefault(Writer, "MipSharpening", BuildSettings.MipSharpening, DefaultSettings.MipSharpening);
 	WriteCbFieldWithDefault(Writer, "DiffuseConvolveMipLevel", BuildSettings.DiffuseConvolveMipLevel, DefaultSettings.DiffuseConvolveMipLevel);
 	WriteCbFieldWithDefault(Writer, "SharpenMipKernelSize", BuildSettings.SharpenMipKernelSize, DefaultSettings.SharpenMipKernelSize);
