TextureCompressorModule increase BuildSourceImageMipsMaxCount to a more future-proof value

diff --git a/Engine/Source/Developer/TextureCompressor/Private/TextureCompressorModule.cpp b/Engine/Source/Developer/TextureCompressor/Private/TextureCompressorModule.cpp
--- a/Engine/Source/Developer/TextureCompressor/Private/TextureCompressorModule.cpp
+++ b/Engine/Source/Developer/TextureCompressor/Private/TextureCompressorModule.cpp
@@ -3262,7 +3262,7 @@ private:
 					);
 
 				// make sure BuildSourceImageMips doesn't reallocate :
-				constexpr int BuildSourceImageMipsMaxCount = 16;
+				constexpr int BuildSourceImageMipsMaxCount = 20; // plenty
 				BuildSourceImageMips.Empty(BuildSourceImageMipsMaxCount);
 
 				// Max Texture Size resizing happens here :
