Clean up AdjustColors in texture compression.
When new mip filters setting is enabled there can be minor differences in output.
Old textures will use previous code and will have no changes in output.

diff --git a/Engine/Source/Developer/TextureCompressor/Private/TextureCompressorModule.cpp b/Engine/Source/Developer/TextureCompressor/Private/TextureCompressorModule.cpp
--- a/Engine/Source/Developer/TextureCompressor/Private/TextureCompressorModule.cpp
+++ b/Engine/Source/Developer/TextureCompressor/Private/TextureCompressorModule.cpp
@@ -2184,8 +2184,7 @@ static bool NeedAdjustImageColors(const FTextureBuildSettings& InBuildSettings)
 		InBuildSettings.bChromaKeyTexture;
 }
 
-static inline void AdjustColors(FLinearColor * Colors,int64 Count,
-		const FTextureBuildSettings& InBuildSettings)
+static inline void AdjustColorsOld(FLinearColor * Colors,int64 Count, const FTextureBuildSettings& InBuildSettings)
 {
 	const FColorAdjustmentParameters& Params = InBuildSettings.ColorAdjustment;
 	
@@ -2347,6 +2346,146 @@ static inline void AdjustColors(FLinearColor * Colors,int64 Count,
 	}
 }
 
+static inline void AdjustColorsNew(FLinearColor* Colors, int64 Count, const FTextureBuildSettings& InBuildSettings)
+{
+	const FColorAdjustmentParameters& Params = InBuildSettings.ColorAdjustment;
+
+	// @todo Oodle : not the same checks as bAdjustNeeded outside
+	//	 but preserves legacy behavior
+	bool bAdjustBrightnessCurve = (!FMath::IsNearlyEqual(Params.AdjustBrightnessCurve, 1.0f, (float)KINDA_SMALL_NUMBER) && Params.AdjustBrightnessCurve != 0.0f);
+	bool bAdjustVibrance = (!FMath::IsNearlyZero(Params.AdjustVibrance, (float)KINDA_SMALL_NUMBER));
+	bool bAdjustRGBCurve = (!FMath::IsNearlyEqual(Params.AdjustRGBCurve, 1.0f, (float)KINDA_SMALL_NUMBER) && Params.AdjustRGBCurve != 0.0f);
+	bool bAdjustSaturation = ( Params.AdjustSaturation != 1.f || bAdjustVibrance );
+	bool bAdjustValue = ( Params.AdjustBrightness != 1.f ) || bAdjustBrightnessCurve;
+
+	float AdjustHue = Params.AdjustHue;
+	if ( AdjustHue != 0.f )
+	{
+		// Params.AdjustHue should be in [0,360] , make sure
+		if ( AdjustHue < 0.f || AdjustHue > 360.f )
+		{
+			AdjustHue = fmodf(AdjustHue, 360.0f);
+			if ( AdjustHue < 0.f )
+			{
+				AdjustHue += 360.f;
+			}
+		}
+	}
+	
+	// BuildSettings.ChromaKeyColor is an FColor
+	FLinearColor ChromaKeyColor(InBuildSettings.ChromaKeyColor);
+
+	bool bChromaKeyTexture = InBuildSettings.bChromaKeyTexture;
+	float ChromaKeyThreshold = InBuildSettings.ChromaKeyThreshold + SMALL_NUMBER;
+
+	bool bHDRSource = InBuildSettings.bHDRSource;
+
+	for (int64 i=0; i<Count; i++)
+	{
+		FLinearColor OriginalColor = Colors[i];
+	
+		if (!bHDRSource)
+		{
+			// Ensure we are clamped as expected (can drift out of clamp due to previous processing)
+			// if you wind up even very slightly out of [0,1] range this function does bad things
+			OriginalColor.R = FMath::Clamp(OriginalColor.R, 0.0f, 1.f);
+			OriginalColor.G = FMath::Clamp(OriginalColor.G, 0.0f, 1.f);
+			OriginalColor.B = FMath::Clamp(OriginalColor.B, 0.0f, 1.f);
+		}
+
+		if (bChromaKeyTexture && (OriginalColor.Equals(ChromaKeyColor, ChromaKeyThreshold)))
+		{
+			Colors[i] = FLinearColor::Transparent;
+			continue;
+		}
+
+		if (bHDRSource)
+		{
+			UE_LOG(LogTextureCompressor, Warning,
+				TEXT("Negative pixel values (%f, %f, %f) are not expected"),
+				OriginalColor.R, OriginalColor.G, OriginalColor.B);
+		}
+
+		// Convert to HSV
+		FLinearColor HSVColor = OriginalColor.LinearRGBToHSV();
+		float& PixelHue = HSVColor.R;
+		float& PixelSaturation = HSVColor.G;
+		float& PixelValue = HSVColor.B;
+
+		if ( bAdjustValue )
+		{
+			// Apply brightness adjustment
+			PixelValue *= Params.AdjustBrightness;
+
+			// Apply brightness power adjustment
+			if ( bAdjustBrightnessCurve )
+			{
+				// Raise HSV.V to the specified power
+				PixelValue = FMath::Pow(PixelValue, Params.AdjustBrightnessCurve);
+			}
+		
+			// Clamp brightness if non-HDR
+			if (!bHDRSource)
+			{
+				PixelValue = FMath::Clamp(PixelValue, 0.0f, 1.0f);
+			}
+		}
+		
+		if ( bAdjustSaturation )
+		{
+			// PixelSaturation is >= 0 but not <= 1
+			//  because negative RGB can come into this function which gives Saturation > 1
+
+			// Apply "vibrance" adjustment
+			if ( bAdjustVibrance )
+			{
+				const float InvSat = 1.0f - PixelSaturation;
+				const float InvSatRaised = InvSat * InvSat * InvSat * InvSat * InvSat;
+
+				const float ClampedVibrance = FMath::Clamp(Params.AdjustVibrance, 0.0f, 1.0f);
+				const float HalfVibrance = ClampedVibrance * 0.5f;
+
+				const float SatProduct = HalfVibrance * InvSatRaised;
+
+				PixelSaturation += SatProduct;
+			}
+
+			// Apply saturation adjustment
+			PixelSaturation *= Params.AdjustSaturation;
+			PixelSaturation = FMath::Clamp(PixelSaturation, 0.0f, 1.0f);
+		}
+
+		// Apply hue adjustment
+		if ( AdjustHue != 0.f )
+		{
+			// PixelHue is [0,360) but AdjustHue is [0,360]
+			PixelHue += AdjustHue;
+
+			// Clamp HSV values
+			if ( PixelHue >= 360.f )
+			{
+				PixelHue -= 360.f;
+			}
+		}
+
+		// Convert back to a linear color
+		FLinearColor LinearColor = HSVColor.HSVToLinearRGB();
+
+		// Apply RGB curve adjustment (linear space)
+		if ( bAdjustRGBCurve )
+		{
+			LinearColor.R = FMath::Pow(LinearColor.R, Params.AdjustRGBCurve);
+			LinearColor.G = FMath::Pow(LinearColor.G, Params.AdjustRGBCurve);
+			LinearColor.B = FMath::Pow(LinearColor.B, Params.AdjustRGBCurve);
+		}
+
+		// Remap the alpha channel
+		LinearColor.A = FMath::Lerp(Params.AdjustMinAlpha, Params.AdjustMaxAlpha, FMath::Clamp(OriginalColor.A, 0.f, 1.f));
+
+		Colors[i] = LinearColor;
+	}
+}
+
 void ITextureCompressorModule::AdjustImageColors(FImage& Image, const FTextureBuildSettings& InBuildSettings)
 {
 	const FColorAdjustmentParameters& InParams = InBuildSettings.ColorAdjustment;
@@ -2383,7 +2522,16 @@ void ITextureCompressorModule::AdjustImageColors(FImage& Image, const FTextureBu
 			FLinearColor * First = &ImageColors[StartIndex];
 			int64 Count = EndIndex-StartIndex;
 
-			AdjustColors(First,Count,InBuildSettings);
+			// Use new AdjustColors code only for newly added textures
+			// So existing textures maintain exactly same output as before
+			if (InBuildSettings.bUseNewMipFilter)
+			{
+				AdjustColorsNew(First, Count, InBuildSettings);
+			}
+			else
+			{
+				AdjustColorsOld(First, Count, InBuildSettings);
+			}
 		}
 		, (bForceSingleThread ? EParallelForFlags::ForceSingleThread : EParallelForFlags::None) );
 	}
@@ -3358,8 +3506,12 @@ private:
 	
 				check( CopyCount == 1 );
 
-				// note the break here skips other adjustments:
-				break;
+				if (!BuildSettings.bUseNewMipFilter)
+				{
+					// note the break here skips other adjustments:
+					// this is maintained only for old textures
+					break;
+				}
 			}
 			else
 			{
