Texture: Fixed an uninitialized build setting that was causing the editor to crash with the new texture build path

#rnx

diff --git a/Engine/Source/Developer/TextureCompressor/Public/TextureCompressorModule.h b/Engine/Source/Developer/TextureCompressor/Public/TextureCompressorModule.h
--- a/Engine/Source/Developer/TextureCompressor/Public/TextureCompressorModule.h
+++ b/Engine/Source/Developer/TextureCompressor/Public/TextureCompressorModule.h
@@ -217,7 +217,8 @@ struct FTextureBuildSettings
 
 	/** Default settings. */
 	FTextureBuildSettings()
-		: bDoScaleMipsForAlphaCoverage(false)
+		: ForceRecompressDDCUID(0)
+		, bDoScaleMipsForAlphaCoverage(false)
 		, AlphaCoverageThresholds(0, 0, 0, 0)
 		, bUseNewMipFilter(false)
 		, MipSharpening(0.0f)
