TextureCompressorModule nop just fix comments

diff --git a/Engine/Source/Developer/TextureCompressor/Private/TextureCompressorModule.cpp b/Engine/Source/Developer/TextureCompressor/Private/TextureCompressorModule.cpp
--- a/Engine/Source/Developer/TextureCompressor/Private/TextureCompressorModule.cpp
+++ b/Engine/Source/Developer/TextureCompressor/Private/TextureCompressorModule.cpp
@@ -2188,7 +2188,10 @@ static inline void AdjustColorsOld(FLinearColor * Colors,int64 Count, const FTex
 {
 	const FColorAdjustmentParameters& Params = InBuildSettings.ColorAdjustment;
 	
-	// @todo Oodle : not the same checks as bAdjustNeeded outside
+	// very similar to AdjustColorsNew
+	// issues in here are mostly fixed in AdjustColorsNew
+
+	// note: : not the same checks as bAdjustNeeded outside
 	//	 but preserves legacy behavior
 	bool bAdjustBrightnessCurve = (!FMath::IsNearlyEqual(Params.AdjustBrightnessCurve, 1.0f, (float)KINDA_SMALL_NUMBER) && Params.AdjustBrightnessCurve != 0.0f);
 	bool bAdjustVibrance = (!FMath::IsNearlyZero(Params.AdjustVibrance, (float)KINDA_SMALL_NUMBER));
@@ -2221,30 +2224,21 @@ static inline void AdjustColorsOld(FLinearColor * Colors,int64 Count, const FTex
 	{
 		FLinearColor OriginalColor = Colors[i];
 	
-		#if 0
-		if ( ! bHDRSource )
-		{
-			// @todo Oodle: for non-HDR source ensure we are clamped as expected
-			//	(can drift out of clamp due to previous processing)
-			//  if you wind up even very slightly out of [0,1] range this function does bad things
-			//	I think this probably should be done, but left disabled to preserve legacy behavior
-			OriginalColor.R = FMath::Clamp(OriginalColor.R, 0.0f, 1.f);
-			OriginalColor.G = FMath::Clamp(OriginalColor.G, 0.0f, 1.f);
-			OriginalColor.B = FMath::Clamp(OriginalColor.B, 0.0f, 1.f);
-		}
-		#endif
+		// note: non-HDR source data in [0,1] can drift outside of [0,1]
+		//   and then bad things will happen here
+		//  left alone to preserve old behavior
 
 		if (bChromaKeyTexture && (OriginalColor.Equals(ChromaKeyColor, ChromaKeyThreshold)))
 		{
 			OriginalColor = FLinearColor::Transparent;
 
-			//@todo Oodle: strange: no return? processing continues on the transparent color...
+			//note: strange: no return.  processing continues on the transparent color...
 			//	  this was likely unintentional
 			//Colors[i] = FLinearColor::Transparent;
 			//continue;
 		}
 
-		// NOTE: if OriginalColor has HDR/floats in it, this does not handle it well
+		// NOTE: if OriginalColor has HDR/floats in it, this function does not handle it well
 		//	it implicitly discards negatives (and negatives cause color shifts)
 		//	for values > 1 the clamp behavior is very strange
 
@@ -2285,7 +2279,6 @@ static inline void AdjustColorsOld(FLinearColor * Colors,int64 Count, const FTex
 			{
 				// note: AdjustVibrance is disabled for HDR source in the Texture UPROPERTIES
 				//    (unclear why, this is no worse than anything else here on HDR)
-				// @todo Oodle : we'd like to just do the multiplies to make pow 5 but can't because it could change floats
 				const float SatRaisePow = 5.0f;
 				const float InvSatRaised = FMath::Pow(1.0f - PixelSaturation, SatRaisePow);
 
@@ -2327,10 +2320,9 @@ static inline void AdjustColorsOld(FLinearColor * Colors,int64 Count, const FTex
 		}
 
 		// Clamp HDR RGB channels to 1 or the original luminance (max original RGB channel value), whichever is greater
-		// @todo Oodle: this is a very odd thing to do
+		// note: this is a very odd thing to do
 		//		clamping at OriginalLuminance if you do AdjustBrightness or AdjustBrightnessCurve ?
 		//	    that would keep values brighter than 1.f unchanged, but bring up lower ones to 1.f
-		//	  I question whether this should just be completely removed
 		if (bHDRSource)
 		{
 			LinearColor.R = FMath::Clamp(LinearColor.R, 0.0f, (OriginalLuminance > 1.0f ? OriginalLuminance : 1.0f));
@@ -2339,13 +2331,13 @@ static inline void AdjustColorsOld(FLinearColor * Colors,int64 Count, const FTex
 		}
 
 		// Remap the alpha channel
-		// @todo Oodle: clamp Original A in [0,1] ?
 		LinearColor.A = FMath::Lerp(Params.AdjustMinAlpha, Params.AdjustMaxAlpha, OriginalColor.A);
 
 		Colors[i] = LinearColor;
 	}
 }
 
+// see also AdjustColorsOld
 static inline void AdjustColorsNew(FLinearColor* Colors, int64 Count, const FTextureBuildSettings& InBuildSettings)
 {
 	const FColorAdjustmentParameters& Params = InBuildSettings.ColorAdjustment;
