Fix Editor crash when performing texture downscaling with specific parameters and texture size.

Note: the crash occured in cases when the scaled dimensions snapped to zero. For example, when the source dimensions are (blocksize * p) and (blocksize * q), where p and q are coprime and downscaling factor scales the dimensions more than twice (considering the rounding).

diff --git a/Engine/Source/Developer/TextureCompressor/Private/TextureCompressorModule.cpp b/Engine/Source/Developer/TextureCompressor/Private/TextureCompressorModule.cpp
--- a/Engine/Source/Developer/TextureCompressor/Private/TextureCompressorModule.cpp
+++ b/Engine/Source/Developer/TextureCompressor/Private/TextureCompressorModule.cpp
@@ -1271,25 +1271,22 @@ static void DownscaleImage(const FImage& SrcImage, FImage& DstImage, const FText
 	}
 	
 	TRACE_CPUPROFILER_EVENT_SCOPE(Texture.DownscaleImage);
-		
+
 	float Downscale = FMath::Clamp(Settings.Downscale, 1.f, 8.f);
+	// note: more accurate would be to use FMath::Max(1, FMath::RoundToInt(SrcImage.SizeX / Downscale))
 	int32 FinalSizeX = FMath::CeilToInt(SrcImage.SizeX / Downscale);
 	int32 FinalSizeY = FMath::CeilToInt(SrcImage.SizeY / Downscale);
 
 	// compute final size respecting image block size
-	if (Settings.BlockSize > 1 
+	if (Settings.BlockSize > 1
 		&& SrcImage.SizeX % Settings.BlockSize == 0 
 		&& SrcImage.SizeY % Settings.BlockSize == 0)
 	{
-		int32 NumBlocksX = SrcImage.SizeX / Settings.BlockSize;
-		int32 NumBlocksY = SrcImage.SizeY / Settings.BlockSize;
-		int32 GCD = FMath::GreatestCommonDivisor(NumBlocksX, NumBlocksY);
-		int32 RatioX = NumBlocksX/GCD;
-		int32 RatioY = NumBlocksY/GCD;
-		int32 FinalNumBlocksX = (int32)FMath::GridSnap((float)FinalSizeX/Settings.BlockSize, (float)RatioX);
-		int32 FinalNumBlocksY = FinalNumBlocksX/RatioX*RatioY;
-		FinalSizeX = FinalNumBlocksX*Settings.BlockSize;
-		FinalSizeY = FinalNumBlocksY*Settings.BlockSize;
+		// the following code finds non-zero dimensions of the scaled image which preserve both aspect ratio and block alignment, and are also the closest to the requested dimensions
+		int32 ScalingGridSizeX = SrcImage.SizeX / FMath::GreatestCommonDivisor(SrcImage.SizeX, SrcImage.SizeY) * Settings.BlockSize;
+		// note: more accurate would be to use (SrcImage.SizeX / Downscale) instead of FinalSizeX here
+		FinalSizeX = FMath::Max(ScalingGridSizeX, FMath::GridSnap(FinalSizeX, ScalingGridSizeX));
+		FinalSizeY = (int64)FinalSizeX * SrcImage.SizeY / SrcImage.SizeX;
 	}
 
 	Downscale = (float)SrcImage.SizeX / FinalSizeX;
