Remove incorrect check from mips generation code.

diff --git a/Engine/Source/Developer/TextureCompressor/Private/TextureCompressorModule.cpp b/Engine/Source/Developer/TextureCompressor/Private/TextureCompressorModule.cpp
--- a/Engine/Source/Developer/TextureCompressor/Private/TextureCompressorModule.cpp
+++ b/Engine/Source/Developer/TextureCompressor/Private/TextureCompressorModule.cpp
@@ -3582,8 +3582,6 @@ private:
 			
 				check( BuildSourceImageMips.Last().SizeX <= MaxTextureResolution &&
 					   BuildSourceImageMips.Last().SizeY <= MaxTextureResolution );
-				check( BuildSourceImageMips.Last().SizeX*2 > MaxTextureResolution ||
-					   BuildSourceImageMips.Last().SizeY*2 > MaxTextureResolution );
 
 				// change pSourceMips to point at the mip chain we made
 				pSourceMips = &BuildSourceImageMips;
