Make DetectAlphaChannel public
#rnx

diff --git a/Engine/Source/Developer/TextureCompressor/Private/TextureCompressorModule.cpp b/Engine/Source/Developer/TextureCompressor/Private/TextureCompressorModule.cpp
--- a/Engine/Source/Developer/TextureCompressor/Private/TextureCompressorModule.cpp
+++ b/Engine/Source/Developer/TextureCompressor/Private/TextureCompressorModule.cpp
@@ -2627,109 +2627,6 @@ static void FlipGreenChannel( FImage& Image )
 	}
 }
 
-/**
- * Detects whether or not the image contains an alpha channel where at least one texel is != 255.
- */
-static bool DetectAlphaChannel(const FImage& InImage)
-{
-	TRACE_CPUPROFILER_EVENT_SCOPE(Texture.DetectAlphaChannel);
-	
-	// previous :
-	// "SMALL_NUMBER" is quite small, this provides almost no tolerance
-	// #define SMALL_NUMBER		(1.e-8f)
-	//const float FloatNonOpaqueAlpha = 1.0f - SMALL_NUMBER;
-	// 
-	// opaque alpha threshold where we'd quantize to < 255 in U8
-	//  images with only alpha larger than this are treated as opaque
-	const float FloatNonOpaqueAlpha = 254.5f / 255.f; // the U8 alpha threshold
-		
-	int64 NumPixels = (int64)InImage.SizeX * InImage.SizeY * InImage.NumSlices;
-
-	if ( InImage.Format == ERawImageFormat::BGRA8 )
-	{
-		TArrayView64<const FColor> SrcColorArray = InImage.AsBGRA8();
-		check( SrcColorArray.Num() == NumPixels );
-
-		const FColor * ColorPtr = &SrcColorArray[0];
-		const FColor * EndPtr = ColorPtr +  SrcColorArray.Num();
-
-		for(;ColorPtr<EndPtr;++ColorPtr)
-		{
-			if ( ColorPtr->A != 255 )
-			{
-				return true;
-			}
-		}
-	}
-	else if ( InImage.Format == ERawImageFormat::RGBA32F )
-	{
-		TArrayView64<const FLinearColor> SrcColorArray = InImage.AsRGBA32F();
-		check( SrcColorArray.Num() == NumPixels );
-
-		const FLinearColor * ColorPtr = &SrcColorArray[0];
-		const FLinearColor * EndPtr = ColorPtr +  SrcColorArray.Num();
-
-		for(;ColorPtr<EndPtr;++ColorPtr)
-		{
-			if (ColorPtr->A <= FloatNonOpaqueAlpha )
-			{
-				return true;
-			}
-		}
-	}
-	else if ( InImage.Format == ERawImageFormat::RGBA16 )
-	{
-		TArrayView64<const uint16> SrcChannelArray = InImage.AsRGBA16();
-		check( SrcChannelArray.Num() == NumPixels*4 );
-
-		const uint16 * ChannelPtr = &SrcChannelArray[0];
-		const uint16 * EndPtr = ChannelPtr +  SrcChannelArray.Num();
-
-		for(;ChannelPtr<EndPtr;ChannelPtr += 4)
-		{
-			if ( ChannelPtr[3] != 0xFFFF )
-			{
-				return true;
-			}
-		}
-	}
-	else if ( InImage.Format == ERawImageFormat::RGBA16F )
-	{
-		TArrayView64<const FFloat16Color> SrcColorArray = InImage.AsRGBA16F();
-		check( SrcColorArray.Num() == NumPixels );
-
-		const FFloat16Color * ColorPtr = &SrcColorArray[0];
-		const FFloat16Color * EndPtr = ColorPtr +  SrcColorArray.Num();
-		
-		for(;ColorPtr<EndPtr;++ColorPtr)
-		{
-			// 16F closest to 1.0 is 0.99951172
-			// use the float tolerance here? or check exactly ?
-			//if ( ColorPtr->A.GetFloat() < 1.f )
-			// use the same FloatNonOpaqueAlpha tolerance for consistency ?
-			if ( ColorPtr->A.GetFloat() < FloatNonOpaqueAlpha )
-			{
-				return true;
-			}
-		}
-	}
-	else if ( InImage.Format == ERawImageFormat::G8 ||
-		InImage.Format == ERawImageFormat::BGRE8 ||
-		InImage.Format == ERawImageFormat::G16 ||
-		InImage.Format == ERawImageFormat::R16F||
-		InImage.Format == ERawImageFormat::R32F )
-	{
-		// source image formats don't have alpha
-	}
-	else
-	{
-		// new format ?
-		check(0);
-	}
-
-	return false;
-}
-
 /** Calculate a scale per 4x4 block of each image, and apply it to the red/green channels. Store scale in the blue channel. */
 static void ApplyYCoCgBlockScale(TArray<FImage>& InOutMipChain)
 {
@@ -3303,7 +3200,7 @@ public:
 		// DetectAlphaChannel on the top mip of the generated mip chain
 		//	BuildSettings could have programatically introduced alpha that was not in the source
 		// note the order of operations in bForceAlphaChannel and bForceNoAlphaChannel ( ForceNo takes precedence )
-		const bool bImageHasAlphaChannel = !BuildSettings.bForceNoAlphaChannel  && (BuildSettings.bForceAlphaChannel || DetectAlphaChannel(IntermediateMipChain[0]));
+		const bool bImageHasAlphaChannel = !BuildSettings.bForceNoAlphaChannel  && (BuildSettings.bForceAlphaChannel || FImageCore::DetectAlphaChannel(IntermediateMipChain[0]));
 	
 		// Set the correct biased texture size so that the compressor understands the original source image size
 		// This is requires for platforms that may need to tile based on the original source texture size
diff --git a/Engine/Source/Runtime/ImageCore/Private/ImageCore.cpp b/Engine/Source/Runtime/ImageCore/Private/ImageCore.cpp
--- a/Engine/Source/Runtime/ImageCore/Private/ImageCore.cpp
+++ b/Engine/Source/Runtime/ImageCore/Private/ImageCore.cpp
@@ -1280,3 +1280,104 @@ void FImageCore::SanitizeFloat16AndSetAlphaOpaqueForBC6H(const FImageView & InOu
 		F16Color.A.SetOne();
 	}
 }
+
+
+bool FImageCore::DetectAlphaChannel(const FImage& InImage)
+{
+	TRACE_CPUPROFILER_EVENT_SCOPE(Texture.DetectAlphaChannel);
+
+	// previous :
+	// "SMALL_NUMBER" is quite small, this provides almost no tolerance
+	// #define SMALL_NUMBER		(1.e-8f)
+	//const float FloatNonOpaqueAlpha = 1.0f - SMALL_NUMBER;
+	// 
+	// opaque alpha threshold where we'd quantize to < 255 in U8
+	//  images with only alpha larger than this are treated as opaque
+	const float FloatNonOpaqueAlpha = 254.5f / 255.f; // the U8 alpha threshold
+
+	int64 NumPixels = (int64)InImage.SizeX * InImage.SizeY * InImage.NumSlices;
+
+	if (InImage.Format == ERawImageFormat::BGRA8)
+	{
+		TArrayView64<const FColor> SrcColorArray = InImage.AsBGRA8();
+		check(SrcColorArray.Num() == NumPixels);
+
+		const FColor* ColorPtr = &SrcColorArray[0];
+		const FColor* EndPtr = ColorPtr + SrcColorArray.Num();
+
+		for (; ColorPtr < EndPtr; ++ColorPtr)
+		{
+			if (ColorPtr->A != 255)
+			{
+				return true;
+			}
+		}
+	}
+	else if (InImage.Format == ERawImageFormat::RGBA32F)
+	{
+		TArrayView64<const FLinearColor> SrcColorArray = InImage.AsRGBA32F();
+		check(SrcColorArray.Num() == NumPixels);
+
+		const FLinearColor* ColorPtr = &SrcColorArray[0];
+		const FLinearColor* EndPtr = ColorPtr + SrcColorArray.Num();
+
+		for (; ColorPtr < EndPtr; ++ColorPtr)
+		{
+			if (ColorPtr->A <= FloatNonOpaqueAlpha)
+			{
+				return true;
+			}
+		}
+	}
+	else if (InImage.Format == ERawImageFormat::RGBA16)
+	{
+		TArrayView64<const uint16> SrcChannelArray = InImage.AsRGBA16();
+		check(SrcChannelArray.Num() == NumPixels * 4);
+
+		const uint16* ChannelPtr = &SrcChannelArray[0];
+		const uint16* EndPtr = ChannelPtr + SrcChannelArray.Num();
+
+		for (; ChannelPtr < EndPtr; ChannelPtr += 4)
+		{
+			if (ChannelPtr[3] != 0xFFFF)
+			{
+				return true;
+			}
+		}
+	}
+	else if (InImage.Format == ERawImageFormat::RGBA16F)
+	{
+		TArrayView64<const FFloat16Color> SrcColorArray = InImage.AsRGBA16F();
+		check(SrcColorArray.Num() == NumPixels);
+
+		const FFloat16Color* ColorPtr = &SrcColorArray[0];
+		const FFloat16Color* EndPtr = ColorPtr + SrcColorArray.Num();
+
+		for (; ColorPtr < EndPtr; ++ColorPtr)
+		{
+			// 16F closest to 1.0 is 0.99951172
+			// use the float tolerance here? or check exactly ?
+			//if ( ColorPtr->A.GetFloat() < 1.f )
+			// use the same FloatNonOpaqueAlpha tolerance for consistency ?
+			if (ColorPtr->A.GetFloat() < FloatNonOpaqueAlpha)
+			{
+				return true;
+			}
+		}
+	}
+	else if (InImage.Format == ERawImageFormat::G8 ||
+		InImage.Format == ERawImageFormat::BGRE8 ||
+		InImage.Format == ERawImageFormat::G16 ||
+		InImage.Format == ERawImageFormat::R16F ||
+		InImage.Format == ERawImageFormat::R32F)
+	{
+		// source image formats don't have alpha
+	}
+	else
+	{
+		// new format ?
+		check(0);
+	}
+
+	return false;
+}
diff --git a/Engine/Source/Runtime/ImageCore/Public/ImageCore.h b/Engine/Source/Runtime/ImageCore/Public/ImageCore.h
--- a/Engine/Source/Runtime/ImageCore/Public/ImageCore.h
+++ b/Engine/Source/Runtime/ImageCore/Public/ImageCore.h
@@ -591,6 +591,15 @@ IMAGECORE_API void TransposeImageRGBABGRA(const FImageView & Image);
 */
 IMAGECORE_API void SanitizeFloat16AndSetAlphaOpaqueForBC6H(const FImageView & InOutImage);
 
+
+/**
+ * Detects whether or not the image contains an alpha channel where at least one texel is != 255.
+ * 
+ * @param InImage - The image to look for alpha in.
+ * @return True if the image both supports and uses the alpha channel. False otherwise.
+ */
+IMAGECORE_API bool DetectAlphaChannel(const FImage & InImage);
+
 };
 
 
