ApplyCompositeTexture allow mismatched composite texture size
fix the case where dest is a partial mip chain, previously failed
fix previously logged warning but made output with no composite
allow nonpow2 or different aspect composite texture
textures where no stretch is required make the same output as before

diff --git a/Engine/Source/Developer/TextureCompressor/Private/TextureCompressorModule.cpp b/Engine/Source/Developer/TextureCompressor/Private/TextureCompressorModule.cpp
--- a/Engine/Source/Developer/TextureCompressor/Private/TextureCompressorModule.cpp
+++ b/Engine/Source/Developer/TextureCompressor/Private/TextureCompressorModule.cpp
@@ -1235,12 +1235,14 @@ static void GenerateTopMip(const FImage& SrcImage, FImage& DestImage, const FTex
 	}
 }
 
+// pixel centers are at XY = integers
+//  range of X is [-0.5,-0.5] to [Width-0.5,Height-0.5]
 static FLinearColor LookupSourceMipBilinear(const FImageView2D& SourceImageData, float X, float Y)
 {
 	X = FMath::Clamp(X, 0.f, SourceImageData.SizeX - 1.f);
 	Y = FMath::Clamp(Y, 0.f, SourceImageData.SizeY - 1.f);
-	int32 IntX0 = FMath::FloorToInt(X);
-	int32 IntY0 = FMath::FloorToInt(Y);
+	int32 IntX0 = FMath::TruncToInt32(X);
+	int32 IntY0 = FMath::TruncToInt32(Y);
 	float FractX = X - IntX0;
 	float FractY = Y - IntY0;
 	int32 IntX1 = FMath::Min(IntX0+1, SourceImageData.SizeX-1);
@@ -1256,6 +1258,45 @@ static FLinearColor LookupSourceMipBilinear(const FImageView2D& SourceImageData,
 	return FMath::Lerp(Sample0, Sample1, FractY);
 }
 
+// UV range is [0,1] , first pixel center is at 0.5/W
+static FLinearColor LookupSourceMipBilinearUV(const FImageView2D& SourceImageData, float U, float V)
+{
+	float X = U * SourceImageData.SizeX - 0.5f;
+	float Y = V * SourceImageData.SizeY - 0.5f;
+	return LookupSourceMipBilinear(SourceImageData,X,Y);
+}
+
+// pixel centers are at XY = integers
+//  range of X is [-0.5,-0.5] to [Width-0.5,Height-0.5]
+static float LookupFloatBilinear(const float * FloatPlane,int32 SizeX,int32 SizeY, float X, float Y)
+{
+	X = FMath::Clamp(X, 0.f, SizeX - 1.f);
+	Y = FMath::Clamp(Y, 0.f, SizeY - 1.f);
+	int32 IntX0 = FMath::TruncToInt32(X);
+	int32 IntY0 = FMath::TruncToInt32(Y);
+	float FractX = X - IntX0;
+	float FractY = Y - IntY0;
+	int32 IntX1 = FMath::Min(IntX0+1, SizeX-1);
+	int32 IntY1 = FMath::Min(IntY0+1, SizeY-1);
+	
+	float Sample00 = FloatPlane[ IntX0 + IntY0 * SizeX ];
+	float Sample10 = FloatPlane[ IntX1 + IntY0 * SizeX ];
+	float Sample01 = FloatPlane[ IntX0 + IntY1 * SizeX ];
+	float Sample11 = FloatPlane[ IntX1 + IntY1 * SizeX ];
+	float Sample0 = FMath::Lerp(Sample00, Sample10, FractX);
+	float Sample1 = FMath::Lerp(Sample01, Sample11, FractX);
+		
+	return FMath::Lerp(Sample0, Sample1, FractY);
+}
+
+// UV range is [0,1] , first pixel center is at 0.5/W
+static float LookupFloatBilinearUV(const float * FloatPlane,int32 SizeX,int32 SizeY, float U, float V)
+{
+	float X = U * SizeX - 0.5f;
+	float Y = V * SizeY - 0.5f;
+	return LookupFloatBilinear(FloatPlane,SizeX,SizeY,X,Y);
+}
+
 struct FTextureDownscaleSettings
 {
 	int32 BlockSize;
@@ -2706,6 +2747,7 @@ static void ApplyYCoCgBlockScale(TArray<FImage>& InOutMipChain)
 	}
 }
 
+#if 0
 static float RoughnessToSpecularPower(float Roughness)
 {
 	float Div = FMath::Pow(Roughness, 4);
@@ -2723,19 +2765,52 @@ static float SpecularPowerToRoughness(float SpecularPower)
 
 	return Out;
 }
+#endif
 
-// @param CompositeTextureMode original type ECompositeTextureMode
-static void ApplyCompositeTexture(FImage& RoughnessSourceMips, const FImage& NormalSourceMips, uint8 CompositeTextureMode, float CompositePower)
+static float CompositeNormalLengthToRoughness(const float LengthN, float Roughness, float CompositePower)
 {
-	TRACE_CPUPROFILER_EVENT_SCOPE(Texture.ApplyCompositeTexture);
+	float Variance = ( 1.0f - LengthN ) / LengthN;
+	Variance = Variance - 0.00004f;
+	if ( Variance <= 0.f )
+	{
+		return Roughness;
+	}
+
+	Variance *= CompositePower;
+		
+#if 0
+	float Power = RoughnessToSpecularPower( Roughness );
+	Power = Power / ( 1.0f + Variance * Power );
+	Roughness = SpecularPowerToRoughness( Power );
+#else
+	// Refactored above to avoid divide by zero
+	float a = Roughness * Roughness;
+	float a2 = a * a;
+	float B = 2.0f * Variance * (a2 - 1.0f);
+	a2 = ( B - a2 ) / ( B - 1.0f );
+	Roughness = FMath::Pow( a2, 0.25f );
+#endif
 
-	check(RoughnessSourceMips.SizeX == NormalSourceMips.SizeX);
-	check(RoughnessSourceMips.SizeY == NormalSourceMips.SizeY);
+	return Roughness;
+}
 
-	FLinearColor* FirstColor = (&RoughnessSourceMips.AsRGBA32F()[0]);
-	const FLinearColor* NormalColors = (&NormalSourceMips.AsRGBA32F()[0]);
+static float CompositeNormalToRoughness(const FLinearColor & NormalColor, float Roughness, float CompositePower)
+{
+	FVector3f Normal = FVector3f(NormalColor.R * 2.0f - 1.0f, NormalColor.G * 2.0f - 1.0f, NormalColor.B * 2.0f - 1.0f);
+				
+	// Toksvig estimation of variance
+	float LengthN = FMath::Min( Normal.Size(), 1.0f );
+
+	return CompositeNormalLengthToRoughness(LengthN,Roughness,CompositePower);
+}
+
+// @param CompositeTextureMode original type ECompositeTextureMode
+//	write roughness to specified channel of DestRoughness, computed from SourceNormals
+static bool ApplyCompositeTexture(FImage& DestRoughness, const FImage& SourceNormals, uint8 CompositeTextureMode, float CompositePower)
+{
+	TRACE_CPUPROFILER_EVENT_SCOPE(Texture.ApplyCompositeTexture);
 
-	int64 Count = (int64) RoughnessSourceMips.SizeX * RoughnessSourceMips.SizeY * RoughnessSourceMips.NumSlices;
+	FLinearColor* FirstColor = (&DestRoughness.AsRGBA32F()[0]);
 	
 	float* TargetValuePtr;
 
@@ -2755,38 +2830,80 @@ static void ApplyCompositeTexture(FImage& RoughnessSourceMips, const FImage& Nor
 			break;
 		default:
 			UE_LOG(LogTextureCompressor, Error, TEXT("Invalid CompositeTextureMode"));
-			return;
+			return false;
 	}
+	
+	if ( DestRoughness.SizeX == SourceNormals.SizeX && DestRoughness.SizeY == SourceNormals.SizeY &&
+		DestRoughness.NumSlices == SourceNormals.NumSlices )
+	{
+		int64 Count = (int64) DestRoughness.SizeX * DestRoughness.SizeY * DestRoughness.NumSlices;
+
+		const FLinearColor* NormalColors = (&SourceNormals.AsRGBA32F()[0]);
 
 	for ( int64 i=0; i<Count; i++ )
 	{
 		const FLinearColor & NormalColor = NormalColors[i];
+
+			float Roughness = TargetValuePtr[i*4];
+			TargetValuePtr[i*4] = CompositeNormalToRoughness(NormalColor,Roughness,CompositePower);
+		}
+	}
+	else
+	{
+		// sizes don't match, stretch with bilinear filter (filter the normal lengths, not the normal colors)
+
+		if ( SourceNormals.NumSlices != 1 ||
+			 DestRoughness.NumSlices != 1 )
+		{
+			UE_LOG(LogTextureCompressor, Warning, TEXT("CompositeTexture XY sizes don't match, stretch not supported because slice count is not 1 : %d vs %d"),
+				SourceNormals.NumSlices,DestRoughness.NumSlices);
+			return false;
+		}
+
+		TArray<float> SourceNormalLengths;
+
+		{
+			int64 SourceNormalCount = (int64) SourceNormals.SizeX * SourceNormals.SizeY;
+			SourceNormalLengths.SetNum(SourceNormalCount);
+
+			const FLinearColor* NormalColors = (&SourceNormals.AsRGBA32F()[0]);
+
+			for ( int64 i=0; i<SourceNormalCount; i++ )
+			{
+				const FLinearColor & NormalColor = NormalColors[i];
+				
 		FVector3f Normal = FVector3f(NormalColor.R * 2.0f - 1.0f, NormalColor.G * 2.0f - 1.0f, NormalColor.B * 2.0f - 1.0f);
 				
-		// Toksvig estimation of variance
 		float LengthN = FMath::Min( Normal.Size(), 1.0f );
-		float Variance = ( 1.0f - LengthN ) / LengthN;
-		Variance = FMath::Max( 0.0f, Variance - 0.00004f );
 
-		Variance *= CompositePower;
-		
-		float Roughness = TargetValuePtr[i*4];
+				SourceNormalLengths[i] = LengthN;
+			}
+		}
 
-#if 0
-		float Power = RoughnessToSpecularPower( Roughness );
-		Power = Power / ( 1.0f + Variance * Power );
-		Roughness = SpecularPowerToRoughness( Power );
-#else
-		// Refactored above to avoid divide by zero
-		float a = Roughness * Roughness;
-		float a2 = a * a;
-		float B = 2.0f * Variance * (a2 - 1.0f);
-		a2 = ( B - a2 ) / ( B - 1.0f );
-		Roughness = FMath::Pow( a2, 0.25f );
-#endif
+		//const FImageView2D NormalColors(const_cast<FImage &>(SourceNormals),0);
+		const float * SourceNormalLengthsPlane = &SourceNormalLengths[0];
+
+		float InvDestW = 1.f/DestRoughness.SizeX;
+		float InvDestH = 1.f/DestRoughness.SizeY;
+
+		for( int64 DestY=0; DestY< DestRoughness.SizeY; DestY++)
+		{
+			float V = (DestY + 0.5f) * InvDestH;
+			for( int64 DestX=0; DestX< DestRoughness.SizeX; DestX++)
+			{
+				float U = (DestX + 0.5f) * InvDestW;
+		
+				//const FLinearColor NormalColor = LookupSourceMipBilinearUV(NormalColors,U,V);
+				const float NormalLength = LookupFloatBilinearUV(SourceNormalLengthsPlane,SourceNormals.SizeX,SourceNormals.SizeY,U,V);
 
-		TargetValuePtr[i*4] = Roughness;
+				float Roughness = *TargetValuePtr;
+				*TargetValuePtr = CompositeNormalLengthToRoughness(NormalLength,Roughness,CompositePower);
+				TargetValuePtr += 4;
+			}
+		}
 	}
+
+	return true;
 }
 
 /*------------------------------------------------------------------------------
@@ -3207,8 +3324,9 @@ public:
 		if (TextureFormat == nullptr)
 		{
 			UE_LOG(LogTextureCompressor, Warning,
-				TEXT("Failed to find compressor for texture format '%s'."),
-				*BuildSettings.TextureFormatName.ToString()
+				TEXT("Failed to find compressor for texture format '%s'. [%.*s]"),
+				*BuildSettings.TextureFormatName.ToString(),
+				DebugTexturePathName.Len(),DebugTexturePathName.GetData()
 			);
 
 			return false;
@@ -3232,9 +3350,8 @@ public:
 		// apply roughness adjustment depending on normal map variation
 		if (AssociatedNormalSourceMips.Num())
 		{
-			// check AssociatedNormalSourceMips.Format; 
 			// ECompositeTextureMode is only NormalRoughness
-			//  composite texture should be a normal map
+			//  AssociatedNormalSourceMips should be a normal map
 
 			TArray<FImage> IntermediateAssociatedNormalSourceMipChain;
 
@@ -3243,28 +3360,34 @@ public:
 			// apply a smooth Gaussian filter to the top level of the normal map
 			// the original comment says :
 			// "helps to reduce aliasing further"
+			// what's happening here is the blur on the top mip will reduce the length of normals in rough areas
+			//	whereas without it the top mip would always have normals of length 1.0 , hence zero roughness per Toksvig
 			DefaultSettings.MipSharpening = -3.5f;
+			//DefaultSettings.MipSharpening = -2.5f; // CB: I think smaller blend looks better but don't change existing data
 			DefaultSettings.SharpenMipKernelSize = 6;
 			DefaultSettings.bApplyKernelToTopMip = true;
 
 			// important to make accurate computation with normal length
+			//  note this normalizes the top mip *before* the gaussian blur
 			DefaultSettings.bRenormalizeTopMip = true;
 
 			// use new mip filter setting from build settings
 			DefaultSettings.bUseNewMipFilter = BuildSettings.bUseNewMipFilter;
 
-			// someday: filtering the normal map then computing roughness is fundamentally wrong
-			//  we should instead compute the roughness scalar first on the original normal map
-			//  then filter on the roughness scalar
-
-			if (!BuildTextureMips(AssociatedNormalSourceMips, DefaultSettings, true, IntermediateAssociatedNormalSourceMipChain, DebugTexturePathName))
+			if (!BuildTextureMips(AssociatedNormalSourceMips, DefaultSettings, CompressorCaps, true, IntermediateAssociatedNormalSourceMipChain))
 			{
-				UE_LOG(LogTextureCompressor, Warning, TEXT("Failed to generate texture mips for composite texture"));
+				UE_LOG(LogTextureCompressor, Warning, TEXT("Failed to generate texture mips for composite texture [%.*s]"),
+					DebugTexturePathName.Len(),DebugTexturePathName.GetData());
+
+				return false;
 			}
 
-			if (!ApplyCompositeTexture(IntermediateMipChain, IntermediateAssociatedNormalSourceMipChain, BuildSettings.CompositeTextureMode, BuildSettings.CompositePower))
+			if (!ApplyCompositeTextureToMips(IntermediateMipChain, IntermediateAssociatedNormalSourceMipChain, BuildSettings.CompositeTextureMode, BuildSettings.CompositePower, BuildSettings.LODBias))
 			{
-				UE_LOG(LogTextureCompressor, Warning, TEXT("Failed to apply composite texture"));
+				UE_LOG(LogTextureCompressor, Warning, TEXT("ApplyCompositeTextureToMips failed [%.*s]"),
+					DebugTexturePathName.Len(),DebugTexturePathName.GetData());
+
+				return false;
 			}
 		}
 
@@ -3728,27 +3851,66 @@ private:
 
 	// @param CompositeTextureMode original type ECompositeTextureMode
 	// @return true on success, false on failure. Can fail due to bad mismatched dimensions of incomplete mip chains.
-	bool ApplyCompositeTexture(TArray<FImage>& RoughnessSourceMips, const TArray<FImage>& NormalSourceMips, uint8 CompositeTextureMode, float CompositePower)
+	bool ApplyCompositeTextureToMips(TArray<FImage>& DestRoughnessMips, const TArray<FImage>& NormalSourceMips, uint8 CompositeTextureMode, float CompositePower, int32 DestLODBias)
 	{
-		uint32 MinLevel = FMath::Min(RoughnessSourceMips.Num(), NormalSourceMips.Num());
+		check( DestRoughnessMips.Num() > 0 );
+		check( NormalSourceMips.Num() > 0 );
+
+		// NormalSourceMips is always a full mip chain, because we just made it, ignoring MipGen settings on the normal map texture
+		// DestRoughnessMips are the mips being made in the current texture build, they could be an incomplete set if NoMipMips or LeaveExisting
+
+		// must write to every mip of output :
+		for(int32 DestLevel = 0; DestLevel < DestRoughnessMips.Num(); ++DestLevel)
+		{
+			if ( DestRoughnessMips[DestLevel].SizeX > NormalSourceMips[0].SizeX &&
+				 DestRoughnessMips[DestLevel].SizeY > NormalSourceMips[0].SizeY )
+			{
+				// Normal map size is smaller than dest, no Roughness needed
+				// at equal size, DO compute roughness as a Gaussian has been applied to filter normals
+				//  to find a roughness at the top mip level
+				continue;
+			}
 
-		if( RoughnessSourceMips[RoughnessSourceMips.Num() - MinLevel].SizeX != NormalSourceMips[NormalSourceMips.Num() - MinLevel].SizeX || 
-			RoughnessSourceMips[RoughnessSourceMips.Num() - MinLevel].SizeY != NormalSourceMips[NormalSourceMips.Num() - MinLevel].SizeY )
+			// find a mip of source normals that is the same size as current mip, if possible
+			int32 SourceNormalMipLevel = FMath::Min(DestLevel,NormalSourceMips.Num()-1);
+			while( SourceNormalMipLevel > 0 && (
+				NormalSourceMips[SourceNormalMipLevel].SizeX < DestRoughnessMips[DestLevel].SizeX ||
+				NormalSourceMips[SourceNormalMipLevel].SizeY < DestRoughnessMips[DestLevel].SizeY) )
+			{
+				SourceNormalMipLevel--;
+			}
+			while( SourceNormalMipLevel < NormalSourceMips.Num()-1 && (
+				NormalSourceMips[SourceNormalMipLevel].SizeX > DestRoughnessMips[DestLevel].SizeX ||
+				NormalSourceMips[SourceNormalMipLevel].SizeY > DestRoughnessMips[DestLevel].SizeY) )
+			{
+				SourceNormalMipLevel++;
+			}
+
+			if ( 
+				DestRoughnessMips[DestLevel].SizeX != NormalSourceMips[SourceNormalMipLevel].SizeX ||
+				DestRoughnessMips[DestLevel].SizeY != NormalSourceMips[SourceNormalMipLevel].SizeY )
 		{
-			UE_LOG(LogTextureCompressor, Warning, TEXT("Couldn't apply composite texture as RoughnessSourceMips (mip %d, %d x %d) doesn't match NormalSourceMips (mip %d, %d x %d); mipchain might be mismatched/incomplete"),
-				RoughnessSourceMips.Num() - MinLevel,
-				RoughnessSourceMips[RoughnessSourceMips.Num() - MinLevel].SizeX,
-				RoughnessSourceMips[RoughnessSourceMips.Num() - MinLevel].SizeY,
-				NormalSourceMips.Num() - MinLevel,
-				NormalSourceMips[NormalSourceMips.Num() - MinLevel].SizeX,
-				NormalSourceMips[NormalSourceMips.Num() - MinLevel].SizeY
+				UE_LOG(LogTextureCompressor, Display, 
+					TEXT( "ApplyCompositeTexture: Couldn't find matching mip size, will stretch.  (dest: %dx%d with %d mips, source: %dx%d with %d mips).  current: (dest: %dx%d at %d, source: %dx%d at %d)" ),
+					DestRoughnessMips[0].SizeX,
+					DestRoughnessMips[0].SizeY,
+					DestRoughnessMips.Num(),
+					NormalSourceMips[0].SizeX,
+					NormalSourceMips[0].SizeY,
+					NormalSourceMips.Num(),
+					DestRoughnessMips[DestLevel].SizeX,
+					DestRoughnessMips[DestLevel].SizeY,
+					DestLevel,
+					NormalSourceMips[SourceNormalMipLevel].SizeX,
+					NormalSourceMips[SourceNormalMipLevel].SizeY,
+					SourceNormalMipLevel
 				);
-			return false;
 		}
 
-		for(uint32 Level = 0; Level < MinLevel; ++Level)
+			if ( ! ApplyCompositeTexture(DestRoughnessMips[DestLevel], NormalSourceMips[SourceNormalMipLevel], CompositeTextureMode, CompositePower) )
 		{
-			::ApplyCompositeTexture(RoughnessSourceMips[RoughnessSourceMips.Num() - 1 - Level], NormalSourceMips[NormalSourceMips.Num() - 1 - Level], CompositeTextureMode, CompositePower);
+				return false;
+			}
 		}
 
 		return true;
diff --git a/Engine/Source/Runtime/Engine/Private/TextureDerivedData.cpp b/Engine/Source/Runtime/Engine/Private/TextureDerivedData.cpp
--- a/Engine/Source/Runtime/Engine/Private/TextureDerivedData.cpp
+++ b/Engine/Source/Runtime/Engine/Private/TextureDerivedData.cpp
@@ -336,7 +336,7 @@ void GetTextureDerivedDataKeySuffix(const UTexture& Texture, const FTextureBuild
 	if(IsValid(Texture.CompositeTexture) && Texture.CompositeTextureMode != CTM_Disabled)
 	{
 		// CompositeTextureMode output changed so force a new DDC key value :
-		CompositeTextureStr += TEXT("_C_");
+		CompositeTextureStr += TEXT("_Composite090802022_");
 		CompositeTextureStr += Texture.CompositeTexture->Source.GetIdString();
 	}
 
@@ -2334,6 +2334,10 @@ static void SerializePlatformData(
 {
 	DECLARE_SCOPE_CYCLE_COUNTER(TEXT("SerializePlatformData"), STAT_Texture_SerializePlatformData, STATGROUP_LoadTime);
 
+	// note: if BuildTexture failed, we still get called here,
+	//	just with a default-constructed PlatformData
+	//	(no mips, sizes=0, PF=Unknown)
+
 	if (Ar.IsFilterEditorOnly())
 	{
 		constexpr int64 PlaceholderDerivedDataSize = 16;
@@ -2505,7 +2509,7 @@ static void SerializePlatformData(
 				FVirtualTextureBuiltData* VTData = PlatformData->VTData;
 				CookTags->Add(Texture, "Size", FString::Printf(TEXT("%dx%d"), VTData->Width, VTData->Height));
 	}
-			else
+			else if ( PlatformData->Mips.Num() > 0 ) // PlatformData->Mips is empty if BuildTexture failed
 			{
 				FString DimensionsStr;
 				FTexture2DMipMap& TopMip = PlatformData->Mips[FirstMipToSerialize];
diff --git a/Engine/Source/Runtime/Engine/Private/TextureDerivedDataTask.cpp b/Engine/Source/Runtime/Engine/Private/TextureDerivedDataTask.cpp
--- a/Engine/Source/Runtime/Engine/Private/TextureDerivedDataTask.cpp
+++ b/Engine/Source/Runtime/Engine/Private/TextureDerivedDataTask.cpp
@@ -1234,45 +1234,15 @@ FTextureCacheDerivedDataWorker::FTextureCacheDerivedDataWorker(
 	if (Texture.CompositeTexture && Texture.CompositeTextureMode != CTM_Disabled)
 	{
 		bool bMatchingBlocks = Texture.CompositeTexture->Source.GetNumBlocks() == Texture.Source.GetNumBlocks();
-		bool bMatchingAspectRatio = true;
-		bool bOnlyPowerOfTwoSize = true;
-		if (bMatchingBlocks)
-		{
-			for (int32 BlockIdx = 0; BlockIdx < Texture.Source.GetNumBlocks(); ++BlockIdx)
-			{
-				FTextureSourceBlock TextureBlock;
-				Texture.Source.GetBlock(BlockIdx, TextureBlock);
-				FTextureSourceBlock CompositeTextureBlock;
-				Texture.CompositeTexture->Source.GetBlock(BlockIdx, CompositeTextureBlock);
-
-				bMatchingBlocks = bMatchingBlocks && TextureBlock.BlockX == CompositeTextureBlock.BlockX && TextureBlock.BlockY == CompositeTextureBlock.BlockY;
-				bMatchingAspectRatio = bMatchingAspectRatio && TextureBlock.SizeX * CompositeTextureBlock.SizeY == TextureBlock.SizeY * CompositeTextureBlock.SizeX;
-				bOnlyPowerOfTwoSize = bOnlyPowerOfTwoSize && FMath::IsPowerOfTwo(TextureBlock.SizeX) && FMath::IsPowerOfTwo(TextureBlock.SizeY);
-			}
-		}
-
+		
 		if (!bMatchingBlocks)
 		{
-			// Only report the warning for textures with a single block
-			// In the future, we should support composite textures if matching blocks are in a different order
-			// Once that's working, then this warning should be reported in all cases
-			if (Texture.Source.GetNumBlocks() == 1)
-			{
-				UE_LOG(LogTexture, Warning, TEXT("Issue while building %s : Composite texture resolution/UDIMs do not match. Composite texture will be ignored"), *TexturePathName);
-			}
-		}
-		else if (!bOnlyPowerOfTwoSize)
-		{
-			UE_LOG(LogTexture, Warning, TEXT("Issue while building %s : Some blocks (UDIMs) have a non power of two size. Composite texture will be ignored"), *TexturePathName);
-		}
-		else if (!bMatchingAspectRatio)
-		{
-			UE_LOG(LogTexture, Warning, TEXT("Issue while building %s : Some blocks (UDIMs) have mismatched aspect ratio. Composite texture will be ignored"), *TexturePathName);
+			UE_LOG(LogTexture, Warning, TEXT("Issue while building %s : Composite texture UDIM Block counts do not match. Composite texture will be ignored"), *TexturePathName);
+			// note: does not fail, fill not warn again
 		}
 
-		if (bMatchingBlocks && bMatchingAspectRatio && bOnlyPowerOfTwoSize)
+		if ( bMatchingBlocks )
 		{
-			// These are derived from the texture, and the composite texture must match.
 			CompositeTextureData.Init(*Texture.CompositeTexture, (TextureMipGenSettings)BuildSettingsPerLayerFetchOrBuild[0].MipGenSettings, BuildSettingsPerLayerFetchOrBuild[0].bCubemap, BuildSettingsPerLayerFetchOrBuild[0].bTextureArray, BuildSettingsPerLayerFetchOrBuild[0].bVolume, bAllowAsyncLoading);
 		}
 	}
@@ -2341,11 +2311,6 @@ public:
 			check( FImageCoreUtils::ConvertToTextureSourceFormat( RawFormat ) == TSF );
 		}
 
-		const bool bCompositeTextureViable = Texture.CompositeTexture && Texture.CompositeTextureMode != CTM_Disabled;
-		bool bMatchingBlocks = bCompositeTextureViable && (Texture.CompositeTexture->Source.GetNumBlocks() == Texture.Source.GetNumBlocks());
-		bool bMatchingAspectRatio = bCompositeTextureViable;
-		bool bOnlyPowerOfTwoSize = bCompositeTextureViable;
-
 		int32 BlockSizeX = 0;
 		int32 BlockSizeY = 0;
 		TArray<FIntPoint> BlockSizes;
@@ -2360,16 +2325,6 @@ public:
 				BlockSizeX = FMath::Max(BlockSizeX, SourceBlock.SizeX);
 				BlockSizeY = FMath::Max(BlockSizeY, SourceBlock.SizeY);
 			}
-
-			if (bCompositeTextureViable)
-			{
-				FTextureSourceBlock CompositeTextureBlock;
-				Texture.CompositeTexture->Source.GetBlock(BlockIndex, CompositeTextureBlock);
-
-				bMatchingBlocks = bMatchingBlocks && SourceBlock.BlockX == CompositeTextureBlock.BlockX && SourceBlock.BlockY == CompositeTextureBlock.BlockY;
-				bMatchingAspectRatio = bMatchingAspectRatio && SourceBlock.SizeX * CompositeTextureBlock.SizeY == SourceBlock.SizeY * CompositeTextureBlock.SizeX;
-				bOnlyPowerOfTwoSize = bOnlyPowerOfTwoSize && FMath::IsPowerOfTwo(SourceBlock.SizeX) && FMath::IsPowerOfTwo(SourceBlock.SizeY);
-			}
 		}
 
 		for (int32 BlockIndex = 0; BlockIndex < BlockSizes.Num(); ++BlockIndex)
@@ -2382,24 +2337,19 @@ public:
 				return false;
 			}
 		}
-
+		
+		const bool bCompositeTextureViable = Texture.CompositeTexture && Texture.CompositeTextureMode != CTM_Disabled;
+		bool bMatchingBlocks = bCompositeTextureViable && (Texture.CompositeTexture->Source.GetNumBlocks() == Texture.Source.GetNumBlocks());
+		
 		if (bCompositeTextureViable)
 		{
 			if (!bMatchingBlocks)
 			{
-				UE_LOG(LogTexture, Warning, TEXT("Issue while building %s : Composite texture resolution/UDIMs do not match. Composite texture will be ignored"), *Texture.GetPathName());
-			}
-			else if (!bOnlyPowerOfTwoSize)
-			{
-				UE_LOG(LogTexture, Warning, TEXT("Issue while building %s : Some blocks (UDIMs) have a non power of two size. Composite texture will be ignored"), *Texture.GetPathName());
-			}
-			else if (!bMatchingAspectRatio)
-			{
-				UE_LOG(LogTexture, Warning, TEXT("Issue while building %s : Some blocks (UDIMs) have mismatched aspect ratio. Composite texture will be ignored"), *Texture.GetPathName());
+				UE_LOG(LogTexture, Warning, TEXT("Issue while building %s : Composite texture UDIM block counts do not match. Composite texture will be ignored"), *Texture.GetPathName());
 			}
 		}
 
-		bOutUseCompositeTexture = bMatchingBlocks && bMatchingAspectRatio && bOnlyPowerOfTwoSize;
+		bOutUseCompositeTexture = bMatchingBlocks;
 
 		// TODO: Add validation equivalent to that found in FTextureCacheDerivedDataWorker::BuildTexture for virtual textures
 		//		 if virtual texture support is added for this code path.
