fix for merge in TextureCompressorModule

diff --git a/Engine/Source/Developer/TextureCompressor/Private/TextureCompressorModule.cpp b/Engine/Source/Developer/TextureCompressor/Private/TextureCompressorModule.cpp
--- a/Engine/Source/Developer/TextureCompressor/Private/TextureCompressorModule.cpp
+++ b/Engine/Source/Developer/TextureCompressor/Private/TextureCompressorModule.cpp
@@ -3374,7 +3374,7 @@ public:
 			// use new mip filter setting from build settings
 			DefaultSettings.bUseNewMipFilter = BuildSettings.bUseNewMipFilter;
 
-			if (!BuildTextureMips(AssociatedNormalSourceMips, DefaultSettings, CompressorCaps, true, IntermediateAssociatedNormalSourceMipChain))
+			if (!BuildTextureMips(AssociatedNormalSourceMips, DefaultSettings, true, IntermediateAssociatedNormalSourceMipChain, DebugTexturePathName))
 			{
 				UE_LOG(LogTextureCompressor, Warning, TEXT("Failed to generate texture mips for composite texture [%.*s]"),
 					DebugTexturePathName.Len(),DebugTexturePathName.GetData());
