Normal Maps were not renormalizing normal vectors after generating mips
caused reduced length normals, which then bends them when we drop components for compression in BC5 or ASTC
don't change existing content, only normalize when new option bNormalizeNormals is enabled
note that ISPCTexComp is inconsistently doing normalization, no other tex formats do

diff --git a/Engine/Source/Developer/TextureBuild/Private/TextureBuildFunction.cpp b/Engine/Source/Developer/TextureBuild/Private/TextureBuildFunction.cpp
--- a/Engine/Source/Developer/TextureBuild/Private/TextureBuildFunction.cpp
+++ b/Engine/Source/Developer/TextureBuild/Private/TextureBuildFunction.cpp
@@ -95,6 +95,7 @@ static FTextureBuildSettings ReadBuildSettingsFromCompactBinary(const FCbObjectV
 	ReadCbField(ColorAdjustmentCbObj["AdjustMinAlpha"], ColorAdjustment.AdjustMinAlpha);
 	ReadCbField(ColorAdjustmentCbObj["AdjustMaxAlpha"], ColorAdjustment.AdjustMaxAlpha);
 	BuildSettings.bUseNewMipFilter = Object["bUseNewMipFilter"].AsBool(BuildSettings.bUseNewMipFilter);
+	BuildSettings.bNormalizeNormals = Object["bNormalizeNormals"].AsBool(BuildSettings.bNormalizeNormals);
 	BuildSettings.bDoScaleMipsForAlphaCoverage = Object["bDoScaleMipsForAlphaCoverage"].AsBool(BuildSettings.bDoScaleMipsForAlphaCoverage);
 	ReadCbField(Object["AlphaCoverageThresholds"], BuildSettings.AlphaCoverageThresholds);
 	ReadCbField(Object["MipSharpening"], BuildSettings.MipSharpening);
diff --git a/Engine/Source/Developer/TextureCompressor/Private/TextureCompressorModule.cpp b/Engine/Source/Developer/TextureCompressor/Private/TextureCompressorModule.cpp
--- a/Engine/Source/Developer/TextureCompressor/Private/TextureCompressorModule.cpp
+++ b/Engine/Source/Developer/TextureCompressor/Private/TextureCompressorModule.cpp
@@ -3089,6 +3089,8 @@ static bool CompressMipChain(
 // only useful for normal maps, fixed bad input (denormalized normals) and improved quality (quantization artifacts)
 static void NormalizeMip(FImage& InOutMip)
 {
+	const FVector NormalIfZero(0.f,0.f,1.f);
+
 	const uint32 NumPixels = InOutMip.SizeX * InOutMip.SizeY * InOutMip.NumSlices;
 	TArrayView64<FLinearColor> ImageColors = InOutMip.AsRGBA32F();
 	for(uint32 CurPixelIndex = 0; CurPixelIndex < NumPixels; ++CurPixelIndex)
@@ -3097,7 +3099,8 @@ static void NormalizeMip(FImage& InOutMip)
 
 		FVector Normal = FVector(Color.R * 2.0f - 1.0f, Color.G * 2.0f - 1.0f, Color.B * 2.0f - 1.0f);
 
-		Normal = Normal.GetSafeNormal();
+		// GetSafeNormal returns Vec(0,0,0) by default for tiny input, instead return flat/up
+		Normal = Normal.GetSafeNormal(UE_SMALL_NUMBER,NormalIfZero);
 
 		Color = FLinearColor(Normal.X * 0.5f + 0.5f, Normal.Y * 0.5f + 0.5f, Normal.Z * 0.5f + 0.5f, Color.A);
 	}
@@ -3370,6 +3373,7 @@ public:
 			// important to make accurate computation with normal length
 			//  note this normalizes the top mip *before* the gaussian blur
 			DefaultSettings.bRenormalizeTopMip = true;
+			DefaultSettings.bNormalizeNormals = false; // do not normalize after mip gen, we want shortening
 
 			// use new mip filter setting from build settings
 			DefaultSettings.bUseNewMipFilter = BuildSettings.bUseNewMipFilter;
@@ -3703,7 +3707,7 @@ private:
 		const bool bLinearize = bNeedLinearize || (GenerateCount > 0) || BuildSettings.bRenormalizeTopMip || (BuildSettings.Downscale > 1.f)
 			|| BuildSettings.bHasColorSpaceDefinition || BuildSettings.bComputeBokehAlpha || BuildSettings.bFlipGreenChannel
 			|| BuildSettings.bReplicateRed || BuildSettings.bReplicateAlpha || BuildSettings.bApplyYCoCgBlockScale
-			|| BuildSettings.SourceEncodingOverride != 0 || bNeedAdjustImageColors;
+			|| BuildSettings.SourceEncodingOverride != 0 || bNeedAdjustImageColors || BuildSettings.bNormalizeNormals;
 
 		for (int32 MipIndex = StartMip; MipIndex < StartMip + CopyCount; ++MipIndex)
 		{
@@ -3748,7 +3752,8 @@ private:
 					}
 					else
 					{
-						// if image is in BGRA8 format leave it, otherwise use original RGBA32F
+						// if image is in BGRA8 format leave it, otherwise convert to RGBA32F
+						//  we only support leaving images in source format if they are BGRA8 and require no processing (eg VT tiles)
 						ERawImageFormat::Type DestFormat = Image.Format == ERawImageFormat::BGRA8 ? ERawImageFormat::BGRA8 : ERawImageFormat::RGBA32F;
 						Image.CopyTo(*Mip, DestFormat, Image.GammaSpace);
 					}
@@ -3833,14 +3838,23 @@ private:
 		}
 
 		// Apply post-mip generation adjustments.
+		if ( BuildSettings.bNormalizeNormals )
+		{
+			for ( FImage& MipImage : OutMipChain )
+			{
+				NormalizeMip(MipImage);
+			}
+		}
+	
 		if (BuildSettings.bReplicateRed)
 		{
 			ReplicateRedChannel(OutMipChain);
 		}
-		else if (BuildSettings.bReplicateAlpha)
+		else if (BuildSettings.bReplicateAlpha) // @@CB why is this an else? 
 		{
 			ReplicateAlphaChannel(OutMipChain);
 		}
+
 		if (BuildSettings.bApplyYCoCgBlockScale)
 		{
 			ApplyYCoCgBlockScale(OutMipChain);
diff --git a/Engine/Source/Developer/TextureCompressor/Public/TextureCompressorModule.h b/Engine/Source/Developer/TextureCompressor/Public/TextureCompressorModule.h
--- a/Engine/Source/Developer/TextureCompressor/Public/TextureCompressorModule.h
+++ b/Engine/Source/Developer/TextureCompressor/Public/TextureCompressorModule.h
@@ -87,6 +87,8 @@ struct FTextureBuildSettings
 	FVector4f AlphaCoverageThresholds;
 	/** Use newer & faster mip generation filter */
 	bool bUseNewMipFilter;
+	/** Normalize normals after mip gen, before compression */
+	bool bNormalizeNormals;
 	/** The desired amount of mip sharpening. */
 	float MipSharpening;
 	/** For angular filtered cubemaps, the mip level which contains convolution with the diffuse cosine lobe. */
@@ -231,6 +233,7 @@ struct FTextureBuildSettings
 		: bDoScaleMipsForAlphaCoverage(false)
 		, AlphaCoverageThresholds(0, 0, 0, 0)
 		, bUseNewMipFilter(false)
+		, bNormalizeNormals(false)
 		, MipSharpening(0.0f)
 		, DiffuseConvolveMipLevel(0)
 		, SharpenMipKernelSize(2)
diff --git a/Engine/Source/Developer/TextureFormatIntelISPCTexComp/Private/TextureFormatIntelISPCTexComp.cpp b/Engine/Source/Developer/TextureFormatIntelISPCTexComp/Private/TextureFormatIntelISPCTexComp.cpp
--- a/Engine/Source/Developer/TextureFormatIntelISPCTexComp/Private/TextureFormatIntelISPCTexComp.cpp
+++ b/Engine/Source/Developer/TextureFormatIntelISPCTexComp/Private/TextureFormatIntelISPCTexComp.cpp
@@ -568,6 +568,8 @@ static void IntelASTCCompressScans(FASTCEncoderSettings* pEncSettings, FImage* p
 			uint8* pInTexelsSwap = pInTexels + (y * InStride);
 			for (int x = 0; x < pInImage->SizeX; ++x)
 			{
+				// @@CB inconsistent : ISPCTexcomp is normalizing normals but nobody else does
+				//	 should remove this and do it at the TextureCompressorModule level so its consistent
 				FVector Normal = FVector(pInTexelsSwap[2] / 255.0f * 2.0f - 1.0f, pInTexelsSwap[1] / 255.0f * 2.0f - 1.0f, pInTexelsSwap[0] / 255.0f * 2.0f - 1.0f);
 				Normal = Normal.GetSafeNormal();
 				pInTexelsSwap[0] = 0;
@@ -589,6 +591,8 @@ static void IntelASTCCompressScans(FASTCEncoderSettings* pEncSettings, FImage* p
 			uint8* pInTexelsSwap = pInTexels + (y * InStride);
 			for (int x = 0; x < pInImage->SizeX; ++x)
 			{
+				// @@CB inconsistent : ISPCTexcomp is normalizing normals but nobody else does
+				//	 should remove this and do it at the TextureCompressorModule level so its consistent
 				FVector Normal = FVector(pInTexelsSwap[2] / 255.0f * 2.0f - 1.0f, pInTexelsSwap[1] / 255.0f * 2.0f - 1.0f, pInTexelsSwap[0] / 255.0f * 2.0f - 1.0f);
 				Normal = Normal.GetSafeNormal();
 				pInTexelsSwap[0] = FMath::RoundToInt((Normal.X * 0.5f + 0.5f) * 255.f);
diff --git a/Engine/Source/Runtime/Engine/Classes/Engine/Texture.h b/Engine/Source/Runtime/Engine/Classes/Engine/Texture.h
--- a/Engine/Source/Runtime/Engine/Classes/Engine/Texture.h
+++ b/Engine/Source/Runtime/Engine/Classes/Engine/Texture.h
@@ -1327,6 +1327,9 @@ public:
 	uint8 SRGB:1;
 
 #if WITH_EDITORONLY_DATA
+	/* Should Normals in Normal Maps be normalized after mip processing?  true gives better quality normals, but false is required to match legacy behavior. */
+	UPROPERTY(EditAnywhere, BlueprintReadWrite, Category=Texture, meta=(DisplayName="Normalize after making mips", EditCondition="CompressionSettings==1"), AdvancedDisplay)
+	uint8 bNormalizeNormals:1;
 
 	/** A flag for using the simplified legacy gamma space e.g pow(color,1/2.2) for converting from FColor to FLinearColor, if we're doing sRGB. */
 	UPROPERTY(EditAnywhere, BlueprintReadWrite, Category=Texture, meta=(DisplayName="sRGB Use Legacy Gamma", EditCondition="SRGB"), AdvancedDisplay)
diff --git a/Engine/Source/Runtime/Engine/Private/Texture.cpp b/Engine/Source/Runtime/Engine/Private/Texture.cpp
--- a/Engine/Source/Runtime/Engine/Private/Texture.cpp
+++ b/Engine/Source/Runtime/Engine/Private/Texture.cpp
@@ -137,6 +137,7 @@ UTexture::UTexture(const FObjectInitializer& ObjectInitializer)
 	//	but it's hard to change now because of UPROPERTY writing deltas to default
 	CompositePower = 1.0f;
 	bUseLegacyGamma = false;
+	bNormalizeNormals = false;
 	bIsImporting = false;
 	bCustomPropertiesImported = false;
 	bDoScaleMipsForAlphaCoverage = false;
@@ -811,9 +812,9 @@ void UTexture::Serialize(FArchive& Ar)
 		Source.EnsureBlocksAreSorted();
 
 		if (Ar.UEVer() < VER_UE4_TEXTURE_LEGACY_GAMMA)
-	{
-		bUseLegacyGamma = true;
-	}
+		{
+			bUseLegacyGamma = true;
+		}
 	}
 
 	if (Ar.IsCooking() && VirtualTextureStreaming)
diff --git a/Engine/Source/Runtime/Engine/Private/TextureDerivedData.cpp b/Engine/Source/Runtime/Engine/Private/TextureDerivedData.cpp
--- a/Engine/Source/Runtime/Engine/Private/TextureDerivedData.cpp
+++ b/Engine/Source/Runtime/Engine/Private/TextureDerivedData.cpp
@@ -275,11 +275,20 @@ static void SerializeForKey(FArchive& Ar, const FTextureBuildSettings& Settings)
 
 	// do not change key if old mip filter is used for old textures
 	// @todo SerializeForKey these can go away whenever we bump the overall ddc key
+	// instead just serialize bool 
 	if (Settings.bUseNewMipFilter)
 	{
 		TempGuid = FGuid(0x27B79A99, 0xE1A5458E, 0xAB619475, 0xCD01AD2A);
 		Ar << TempGuid;
 	}
+	
+	// @todo SerializeForKey these can go away whenever we bump the overall ddc key
+	// instead just serialize bool bNormalizeNormals
+	if ( Settings.bNormalizeNormals )
+	{
+		TempGuid = FGuid(0x0F5221F6,0x992344D3,0x9C3CCED9,0x4AF08FB8);
+		Ar << TempGuid;
+	}
 
 	if (Settings.bLongLatSource)
 	{
@@ -832,6 +841,7 @@ static void GetTextureBuildSettings(
 
 	OutBuildSettings.CompressionCacheId = Texture.CompressionCacheId;
 	OutBuildSettings.bUseNewMipFilter = Texture.bUseNewMipFilter;
+	OutBuildSettings.bNormalizeNormals = Texture.bNormalizeNormals && Texture.IsNormalMap();
 	OutBuildSettings.bComputeBokehAlpha = (Texture.LODGroup == TEXTUREGROUP_Bokeh);
 	OutBuildSettings.bReplicateAlpha = false;
 	OutBuildSettings.bReplicateRed = false;
diff --git a/Engine/Source/Runtime/Engine/Private/TextureDerivedDataBuildUtils.cpp b/Engine/Source/Runtime/Engine/Private/TextureDerivedDataBuildUtils.cpp
--- a/Engine/Source/Runtime/Engine/Private/TextureDerivedDataBuildUtils.cpp
+++ b/Engine/Source/Runtime/Engine/Private/TextureDerivedDataBuildUtils.cpp
@@ -130,6 +130,7 @@ static void WriteBuildSettings(FCbWriter& Writer, const FTextureBuildSettings& B
 	}
 	WriteCbFieldWithDefault(Writer, "CompressionCacheId", BuildSettings.CompressionCacheId, DefaultSettings.CompressionCacheId);
 	WriteCbFieldWithDefault<bool>(Writer, "bUseNewMipFilter", BuildSettings.bUseNewMipFilter, DefaultSettings.bUseNewMipFilter);
+	WriteCbFieldWithDefault<bool>(Writer, "bNormalizeNormals", BuildSettings.bNormalizeNormals, DefaultSettings.bNormalizeNormals);
 	WriteCbFieldWithDefault(Writer, "MipSharpening", BuildSettings.MipSharpening, DefaultSettings.MipSharpening);
 	WriteCbFieldWithDefault(Writer, "DiffuseConvolveMipLevel", BuildSettings.DiffuseConvolveMipLevel, DefaultSettings.DiffuseConvolveMipLevel);
 	WriteCbFieldWithDefault(Writer, "SharpenMipKernelSize", BuildSettings.SharpenMipKernelSize, DefaultSettings.SharpenMipKernelSize);
