TextureCompressorModule.cpp nop comment on Downscale

diff --git a/Engine/Source/Developer/TextureCompressor/Private/TextureCompressorModule.cpp b/Engine/Source/Developer/TextureCompressor/Private/TextureCompressorModule.cpp
--- a/Engine/Source/Developer/TextureCompressor/Private/TextureCompressorModule.cpp
+++ b/Engine/Source/Developer/TextureCompressor/Private/TextureCompressorModule.cpp
@@ -1308,7 +1308,7 @@ struct FTextureDownscaleSettings
 	{
 		Downscale = BuildSettings.Downscale;
 		DownscaleOptions = BuildSettings.DownscaleOptions;
-		BlockSize = 4;
+		BlockSize = 4; // <- hard coded to 4 even if texture is not compressed
 		UseNewMipFilter = BuildSettings.bUseNewMipFilter;
 	}
 };
@@ -1316,21 +1316,50 @@ struct FTextureDownscaleSettings
 static float GetDownscaleFinalSizeAndClampedDownscale(int32 SrcImageWidth, int32 SrcImageHeight,  const FTextureDownscaleSettings& Settings, int32& OutWidth, int32& OutHeight)
 {
 	check(Settings.Downscale > 1.0f); // must be already handled.
+	
 	float Downscale = FMath::Clamp(Settings.Downscale, 1.f, 8.f);
+
 	// note: more accurate would be to use FMath::Max(1, FMath::RoundToInt(SrcImage.SizeX / Downscale))
 	int32 FinalSizeX = FMath::CeilToInt(SrcImageWidth / Downscale);
 	int32 FinalSizeY = FMath::CeilToInt(SrcImageHeight / Downscale);
 
 	// compute final size respecting image block size
 	if (Settings.BlockSize > 1
-		&& SrcImageWidth % Settings.BlockSize == 0
-		&& SrcImageHeight % Settings.BlockSize == 0)
+		&& (SrcImageWidth % Settings.BlockSize) == 0
+		&& (SrcImageHeight % Settings.BlockSize) == 0)
 	{
-		// the following code finds non-zero dimensions of the scaled image which preserve both aspect ratio and block alignment, and are also the closest to the requested dimensions
-		int32 ScalingGridSizeX = SrcImageWidth / FMath::GreatestCommonDivisor(SrcImageWidth, SrcImageHeight) * Settings.BlockSize;
+		// the following code finds non-zero dimensions of the scaled image which preserve both aspect ratio and block alignment, 
+		// it favors preserving aspect ratio at the expense of not scaling the desired factor when both are not possible
+		int32 GCD = FMath::GreatestCommonDivisor(SrcImageWidth, SrcImageHeight);
+		int32 ScalingGridSizeX = (SrcImageWidth / GCD) * Settings.BlockSize;
 		// note: more accurate would be to use (SrcImage.SizeX / Downscale) instead of FinalSizeX here
-		FinalSizeX = FMath::Max(ScalingGridSizeX, FMath::GridSnap(FinalSizeX, ScalingGridSizeX));
-		FinalSizeY = (int64)FinalSizeX * SrcImageHeight / SrcImageWidth;
+		// GridSnap rounds to nearest, and can return zero
+		FinalSizeX = FMath::GridSnap(FinalSizeX, ScalingGridSizeX);
+		FinalSizeX = FMath::Max(ScalingGridSizeX, FinalSizeX);
+		FinalSizeY = (int32)( ((int64)FinalSizeX * SrcImageHeight) / SrcImageWidth );
+		// Final Size X and Y are gauranteed to be block aligned
+
+		#if 0
+		
+		// simpler alternative :
+		// choose the block count in the smaller dimension first
+		// then make the larger dimension maintain aspect ratio
+		int32 FinalNumBlocksX,FinalNumBlocksY;
+		if ( SrcImage.SizeX >= SrcImage.SizeY )
+		{
+			FinalNumBlocksY = FMath::RoundToInt( SrcImage.SizeY / (Downscale * Settings.BlockSize) );
+			FinalNumBlocksX = FMath::RoundToInt( FinalNumBlocksY * SrcImage.SizeX / (float)SrcImage.SizeY );
+		}
+		else
+		{
+			FinalNumBlocksX = FMath::RoundToInt( SrcImage.SizeX / (Downscale * Settings.BlockSize) );
+			FinalNumBlocksY = FMath::RoundToInt( FinalNumBlocksX * SrcImage.SizeY / (float)SrcImage.SizeX );
+		}
+
+		FinalSizeX = FMath::Max(FinalNumBlocksX,1)*Settings.BlockSize;
+		FinalSizeY = FMath::Max(FinalNumBlocksY,1)*Settings.BlockSize;
+
+		#endif
 	}
 
 	OutWidth = FinalSizeX;
@@ -1345,13 +1374,21 @@ static void DownscaleImage(const FImage& SrcImage, FImage& DstImage, const FText
 		return;
 	}
 	
+	// Settings.BlockSize == 4 always, even if not compressed
+	// Downscale can only be applied if NoMipMaps, otherwise it is silently ignored
+	//	 also only applied if Texture2D , ignored by all other texture types
+
 	TRACE_CPUPROFILER_EVENT_SCOPE(Texture.DownscaleImage);
 
 	int32 FinalSizeX = 0, FinalSizeY =0;
 	float Downscale = GetDownscaleFinalSizeAndClampedDownscale(SrcImage.SizeX, SrcImage.SizeY, Settings, FinalSizeX, FinalSizeY);
 
+	// recompute Downscale factor because it may have changed due to block alignment
+	// note: if aspect ratio was not exactly preserved, this could differ in X and Y
 	Downscale = (float)SrcImage.SizeX / FinalSizeX;
-		
+	
+	// while desired final size is > 2X smaller, do 2X resizes using the standard mip gen filter code
+
 	FImage Image0;
 	FImage Image1;
 	FImage* ImageChain[2] = {&const_cast<FImage&>(SrcImage), &Image1};
@@ -1427,11 +1464,15 @@ static void DownscaleImage(const FImage& SrcImage, FImage& DstImage, const FText
 	}
 	
 	ImageChain[1]->Init(FinalSizeX, FinalSizeY, ImageChain[0]->NumSlices, ImageChain[0]->Format, ImageChain[0]->GammaSpace);
+
+	// recompute Downscale factor again for final arbitrary-factor resizes :
+	// note: if aspect ratio was not exactly preserved, this could differ in X and Y
 	Downscale = (float)ImageChain[0]->SizeX / FinalSizeX;
 
 	FImageView2D SrcImageData(*ImageChain[0], 0);
 	FImageView2D DstImageData(*ImageChain[1], 0);
 					
+	// @todo Oodle : not sure this is a correct image resize without shift; does it get pixel center offsets right?
 	for (int32 Y = 0; Y < FinalSizeY; ++Y)
 	{
 		float SourceY = Y * Downscale;
@@ -3848,9 +3889,10 @@ private:
 	
 		if (BuildSettings.bReplicateRed)
 		{
+			check( !BuildSettings.bReplicateAlpha ); // cannot both be set 
 			ReplicateRedChannel(OutMipChain);
 		}
-		else if (BuildSettings.bReplicateAlpha) // @@CB why is this an else? 
+		else if (BuildSettings.bReplicateAlpha)
 		{
 			ReplicateAlphaChannel(OutMipChain);
 		}
