Fix for cubemap arrays built from longlat sources predicting incorrect mip slice count.

diff --git a/Engine/Source/Developer/TextureCompressor/Private/TextureCompressorModule.cpp b/Engine/Source/Developer/TextureCompressor/Private/TextureCompressorModule.cpp
--- a/Engine/Source/Developer/TextureCompressor/Private/TextureCompressorModule.cpp
+++ b/Engine/Source/Developer/TextureCompressor/Private/TextureCompressorModule.cpp
@@ -3333,7 +3333,7 @@ public:
 			uint32 LongLatCubemapExtents = ComputeLongLatCubemapExtents(BaseSizeX, BuildSettings.MaxTextureResolution);
 			BaseSizeX = LongLatCubemapExtents;
 			BaseSizeY = LongLatCubemapExtents;
-			OutMip0NumSlices = 6;
+			OutMip0NumSlices = 6 * InMip0NumSlices;
 		}
 
 		// At this point we have a base mip size that is valid.
