fix TextureCompressorModule incorrect "Negative pixel values" Warning spew

diff --git a/Engine/Source/Developer/TextureCompressor/Private/TextureCompressorModule.cpp b/Engine/Source/Developer/TextureCompressor/Private/TextureCompressorModule.cpp
--- a/Engine/Source/Developer/TextureCompressor/Private/TextureCompressorModule.cpp
+++ b/Engine/Source/Developer/TextureCompressor/Private/TextureCompressorModule.cpp
@@ -2465,7 +2465,6 @@ static inline void AdjustColorsNew(FLinearColor* Colors, int64 Count, const FTex
 	
 	// BuildSettings.ChromaKeyColor is an FColor
 	FLinearColor ChromaKeyColor(InBuildSettings.ChromaKeyColor);
-
 	bool bChromaKeyTexture = InBuildSettings.bChromaKeyTexture;
 	float ChromaKeyThreshold = InBuildSettings.ChromaKeyThreshold + SMALL_NUMBER;
 
@@ -2473,8 +2472,14 @@ static inline void AdjustColorsNew(FLinearColor* Colors, int64 Count, const FTex
 
 	for (int64 i=0; i<Count; i++)
 	{
+		if (bChromaKeyTexture && (Colors[i].Equals(ChromaKeyColor, ChromaKeyThreshold)))
+		{
+			Colors[i] = FLinearColor::Transparent;
+			continue;
+		}
+
 		FLinearColor OriginalColor = Colors[i];
-	
+			
 		if (!bHDRSource)
 		{
 			// Ensure we are clamped as expected (can drift out of clamp due to previous processing)
@@ -2483,18 +2488,25 @@ static inline void AdjustColorsNew(FLinearColor* Colors, int64 Count, const FTex
 			OriginalColor.G = FMath::Clamp(OriginalColor.G, 0.0f, 1.f);
 			OriginalColor.B = FMath::Clamp(OriginalColor.B, 0.0f, 1.f);
 		}
-
-		if (bChromaKeyTexture && (OriginalColor.Equals(ChromaKeyColor, ChromaKeyThreshold)))
-		{
-			Colors[i] = FLinearColor::Transparent;
-			continue;
-		}
-
-		if (bHDRSource)
+		else
 		{
-			UE_LOG(LogTextureCompressor, Warning,
-				TEXT("Negative pixel values (%f, %f, %f) are not expected"),
-				OriginalColor.R, OriginalColor.G, OriginalColor.B);
+			/*
+			if ( OriginalColor.R < 0 || OriginalColor.G < 0 || OriginalColor.B < 0 )
+			{
+				// need to log texture name for this to be of any use :
+				UE_CALL_ONCE( [&](){
+					UE_LOG(LogTextureCompressor, Warning,
+						TEXT("Negative pixel values (%f, %f, %f) are not expected"),
+						OriginalColor.R, OriginalColor.G, OriginalColor.B);
+				} );
+			}
+			*/
+			
+			// yes HDR source, but we don't support negatives in Adjust
+			// clamp to >= 0 :
+			OriginalColor.R = FMath::Max(OriginalColor.R, 0.0f);
+			OriginalColor.G = FMath::Max(OriginalColor.G, 0.0f);
+			OriginalColor.B = FMath::Max(OriginalColor.B, 0.0f);
 		}
 
 		// Convert to HSV
