TextureCompressorModule: LongLat to Cubemap conversion parallelize on 6 faces
and simplify index clamp logic

diff --git a/Engine/Source/Developer/TextureCompressor/Private/TextureCompressorModule.cpp b/Engine/Source/Developer/TextureCompressor/Private/TextureCompressorModule.cpp
--- a/Engine/Source/Developer/TextureCompressor/Private/TextureCompressorModule.cpp
+++ b/Engine/Source/Developer/TextureCompressor/Private/TextureCompressorModule.cpp
@@ -1235,6 +1235,39 @@ static void GenerateTopMip(const FImage& SrcImage, FImage& DestImage, const FTex
 	}
 }
 
+static FLinearColor Bilerp(
+	const FLinearColor & Sample00,
+	const FLinearColor & Sample10,
+	const FLinearColor & Sample01,
+	const FLinearColor & Sample11,
+	float FracX,float FracY)
+{
+	//FMath::Lerp is :
+	// return (T)(A + Alpha * (B-A));
+	// must match that exactly
+
+	VectorRegister4Float V00 = VectorLoad(&Sample00.Component(0));
+	VectorRegister4Float V10 = VectorLoad(&Sample10.Component(0));
+	VectorRegister4Float V01 = VectorLoad(&Sample01.Component(0));
+	VectorRegister4Float V11 = VectorLoad(&Sample11.Component(0));
+
+	VectorRegister4Float S0 = VectorAdd(V00, VectorMultiply(VectorSetFloat1(FracX), VectorSubtract(V10,V00)));
+	VectorRegister4Float S1 = VectorAdd(V01, VectorMultiply(VectorSetFloat1(FracX), VectorSubtract(V11,V01)));
+	VectorRegister4Float SS = VectorAdd(S0 , VectorMultiply(VectorSetFloat1(FracY), VectorSubtract(S1 ,S0 )));
+	FLinearColor Out;
+	VectorStore(SS,&Out.Component(0));
+	
+	/*
+	// FLinearColor has math operators but they are scalar, not vector
+	FLinearColor Sample0 = FMath::Lerp(Sample00, Sample10, FracX);
+	FLinearColor Sample1 = FMath::Lerp(Sample01, Sample11, FracX);
+	FLinearColor Ret = FMath::Lerp(Sample0, Sample1, FracY);
+	check( Out == Ret );
+	*/
+
+	return Out;
+}
+
 // pixel centers are at XY = integers
 //  range of X is [-0.5,-0.5] to [Width-0.5,Height-0.5]
 static FLinearColor LookupSourceMipBilinear(const FImageView2D& SourceImageData, float X, float Y)
@@ -1248,14 +1281,12 @@ static FLinearColor LookupSourceMipBilinear(const FImageView2D& SourceImageData,
 	int32 IntX1 = FMath::Min(IntX0+1, SourceImageData.SizeX-1);
 	int32 IntY1 = FMath::Min(IntY0+1, SourceImageData.SizeY-1);
 	
-	FLinearColor Sample00 = SourceImageData.Access(IntX0,IntY0);
-	FLinearColor Sample10 = SourceImageData.Access(IntX1,IntY0);
-	FLinearColor Sample01 = SourceImageData.Access(IntX0,IntY1);
-	FLinearColor Sample11 = SourceImageData.Access(IntX1,IntY1);
-	FLinearColor Sample0 = FMath::Lerp(Sample00, Sample10, FractX);
-	FLinearColor Sample1 = FMath::Lerp(Sample01, Sample11, FractX);
-		
-	return FMath::Lerp(Sample0, Sample1, FractY);
+	const FLinearColor & Sample00 = SourceImageData.Access(IntX0,IntY0);
+	const FLinearColor & Sample10 = SourceImageData.Access(IntX1,IntY0);
+	const FLinearColor & Sample01 = SourceImageData.Access(IntX0,IntY1);
+	const FLinearColor & Sample11 = SourceImageData.Access(IntX1,IntY1);
+
+	return Bilerp(Sample00,Sample10,Sample01,Sample11,FractX,FractY);
 }
 
 // UV range is [0,1] , first pixel center is at 0.5/W
@@ -1748,22 +1779,11 @@ struct FImageViewLongLat
 	{
 		SizeX = Image.SizeX;
 		SizeY = Image.SizeY;
-		ImageColors = (&Image.AsRGBA32F()[0]) + SliceIndex * SizeY * SizeX;
-	}
-
-	/** Wraps X around W. */
-	static void WrapTo(int32& X, int32 W)
-	{
-		X = X % W;
-
-		if(X < 0)
-		{
-			X += W;
-		}
+		ImageColors = (&Image.AsRGBA32F()[0]) + SliceIndex * (int64) SizeY * SizeX;
 	}
 
 	/** Const access to a texel. */
-	FLinearColor Access(int32 X, int32 Y) const
+	const FLinearColor & Access(int32 X, int32 Y) const
 	{
 		return ImageColors[X + Y * SizeX];
 	}
@@ -1771,8 +1791,13 @@ struct FImageViewLongLat
 	/** Makes a filtered lookup. */
 	FLinearColor LookupFiltered(float X, float Y) const
 	{
-		int32 X0 = (int32)floorf(X);
-		int32 Y0 = (int32)floorf(Y);
+		// X is in (0,SizeX) (exclusive)
+		// Y is in (0,SizeY)
+		checkSlow( X > 0.f && X < SizeX );
+		checkSlow( Y > 0.f && Y < SizeY );
+
+		int32 X0 = (int32)X;
+		int32 Y0 = (int32)Y;
 
 		float FracX = X - X0;
 		float FracY = Y - Y0;
@@ -1780,28 +1805,52 @@ struct FImageViewLongLat
 		int32 X1 = X0 + 1;
 		int32 Y1 = Y0 + 1;
 
-		WrapTo(X0, SizeX);
-		WrapTo(X1, SizeX);
-		Y0 = FMath::Clamp(Y0, 0, (int32)(SizeY - 1));
-		Y1 = FMath::Clamp(Y1, 0, (int32)(SizeY - 1));
+		// wrap X :
+		checkSlow( X0 >= 0 && X0 < SizeX );
+		if ( X1 >= SizeX )
+		{
+			X1 = 0;
+		}
 
-		FLinearColor CornerRGB00 = Access(X0, Y0);
-		FLinearColor CornerRGB10 = Access(X1, Y0);
-		FLinearColor CornerRGB01 = Access(X0, Y1);
-		FLinearColor CornerRGB11 = Access(X1, Y1);
+		// clamp Y :
+		// clamp should only ever change SizeY to SizeY -1 ?
+		checkSlow( Y0 >= 0 && Y0 < SizeY );
+		if ( Y1 >= SizeY )
+		{
+			Y1 = SizeY-1;
+		}
 
-		FLinearColor CornerRGB0 = FMath::Lerp(CornerRGB00, CornerRGB10, FracX);
-		FLinearColor CornerRGB1 = FMath::Lerp(CornerRGB01, CornerRGB11, FracX);
+		const FLinearColor & CornerRGB00 = Access(X0, Y0);
+		const FLinearColor & CornerRGB10 = Access(X1, Y0);
+		const FLinearColor & CornerRGB01 = Access(X0, Y1);
+		const FLinearColor & CornerRGB11 = Access(X1, Y1);
 
-		return FMath::Lerp(CornerRGB0, CornerRGB1, FracY);
+		return Bilerp(CornerRGB00,CornerRGB10,CornerRGB01,CornerRGB11,FracX,FracY);
 	}
 
 	/** Makes a filtered lookup using a direction. */
-	FLinearColor LookupLongLat(FVector NormalizedDirection) const
+	// see http://gl.ict.usc.edu/Data/HighResProbes
+	// latitude-longitude panoramic format = equirectangular mapping
+
+	// using floats saves ~6% of the total time
+	//	but would change output
+	// cycles_float = 135 cycles_double = 143
+	FLinearColor LookupLongLatFloat(const FVector & NormalizedDirection) const
 	{
-		// see http://gl.ict.usc.edu/Data/HighResProbes
-		// latitude-longitude panoramic format = equirectangular mapping
+		// atan2 returns in [-PI,PI]
+		// acos returns in [0,PI]
 
+		const float invPI = 1.f/PI;
+		float X = (1.f + atan2f(NormalizedDirection.X, - NormalizedDirection.Z) * invPI) * 0.5f * SizeX;
+		float Y = acosf(NormalizedDirection.Y)*invPI * SizeY;
+
+		return LookupFiltered(X, Y);
+	}
+	
+	FLinearColor LookupLongLatDouble(const FVector & NormalizedDirection) const
+	{
+		// this does the math in doubles then stores to floats :
+		//	that was probably a mistake, but leave it to avoid patches
 		float X = (1 + atan2(NormalizedDirection.X, - NormalizedDirection.Z) / PI) / 2 * SizeX;
 		float Y = acos(NormalizedDirection.Y) / PI * SizeY;
 
@@ -1810,23 +1859,22 @@ struct FImageViewLongLat
 };
 
 // transform world space vector to a space relative to the face
-static FVector TransformSideToWorldSpace(uint32 CubemapFace, FVector InDirection)
+static inline FVector TransformSideToWorldSpace(uint32 CubemapFace, const FVector & InDirection)
 {
 	float x = InDirection.X, y = InDirection.Y, z = InDirection.Z;
 
-	FVector Ret = FVector(0, 0, 0);
+	FVector Ret;
 
 	// see http://msdn.microsoft.com/en-us/library/bb204881(v=vs.85).aspx
 	switch(CubemapFace)
 	{
+		default: checkSlow(0);
 		case 0: Ret = FVector(+z, -y, -x); break;
 		case 1: Ret = FVector(-z, -y, +x); break;
 		case 2: Ret = FVector(+x, +z, +y); break;
 		case 3: Ret = FVector(+x, -z, -y); break;
 		case 4: Ret = FVector(+x, -y, +z); break;
 		case 5: Ret = FVector(-x, -y, -z); break;
-		default:
-			checkSlow(0);
 	}
 
 	// this makes it with the Unreal way (z and y are flipped)
@@ -1834,30 +1882,29 @@ static FVector TransformSideToWorldSpace(uint32 CubemapFace, FVector InDirection
 }
 
 // transform vector relative to the face to world space
-static FVector TransformWorldToSideSpace(uint32 CubemapFace, FVector InDirection)
+static inline FVector TransformWorldToSideSpace(uint32 CubemapFace, const FVector & InDirection)
 {
 	// undo Unreal way (z and y are flipped)
 	float x = InDirection.X, y = InDirection.Z, z = InDirection.Y;
 
-	FVector Ret = FVector(0, 0, 0); 
+	FVector Ret;
 
 	// see http://msdn.microsoft.com/en-us/library/bb204881(v=vs.85).aspx
 	switch(CubemapFace)
 	{
+		default: checkSlow(0);
 		case 0: Ret = FVector(-z, -y, +x); break;
 		case 1: Ret = FVector(+z, -y, -x); break;
 		case 2: Ret = FVector(+x, +z, +y); break;
 		case 3: Ret = FVector(+x, -z, -y); break;
 		case 4: Ret = FVector(+x, -y, +z); break;
 		case 5: Ret = FVector(-x, -y, -z); break;
-		default:
-			checkSlow(0);
 	}
 
 	return Ret;
 }
 
-static FVector ComputeSSCubeDirectionAtTexelCenter(uint32 x, uint32 y, float InvSideExtent)
+static inline FVector ComputeSSCubeDirectionAtTexelCenter(uint32 x, uint32 y, float InvSideExtent)
 {
 	// center of the texels
 	FVector DirectionSS((x + 0.5f) * InvSideExtent * 2 - 1, (y + 0.5f) * InvSideExtent * 2 - 1, 1);
@@ -1865,7 +1912,7 @@ static FVector ComputeSSCubeDirectionAtTexelCenter(uint32 x, uint32 y, float Inv
 	return DirectionSS;
 }
 
-static FVector ComputeWSCubeDirectionAtTexelCenter(uint32 CubemapFace, uint32 x, uint32 y, float InvSideExtent)
+static inline FVector ComputeWSCubeDirectionAtTexelCenter(uint32 CubemapFace, uint32 x, uint32 y, float InvSideExtent)
 {
 	FVector DirectionSS = ComputeSSCubeDirectionAtTexelCenter(x, y, InvSideExtent);
 	FVector DirectionWS = TransformSideToWorldSpace(CubemapFace, DirectionSS);
@@ -1879,7 +1926,7 @@ static uint32 ComputeLongLatCubemapExtents(int32 SrcImageSizeX, const uint32 Max
 
 void ITextureCompressorModule::GenerateBaseCubeMipFromLongitudeLatitude2D(FImage* OutMip, const FImage& SrcImage, const uint32 MaxCubemapTextureResolution, uint8 SourceEncodingOverride)
 {
-	TRACE_CPUPROFILER_EVENT_SCOPE(Texture.GenerateBaseCubeMipFromLongitudeLatitude2D);
+	TRACE_CPUPROFILER_EVENT_SCOPE(Texture.CubeMipFromLongLat);
 
 	FImage LongLatImage;
 	SrcImage.Linearize(SourceEncodingOverride, LongLatImage);
@@ -1892,7 +1939,9 @@ void ITextureCompressorModule::GenerateBaseCubeMipFromLongitudeLatitude2D(FImage
 	for (int32 Slice = 0; Slice < SrcImage.NumSlices; ++Slice)
 	{
 		FImageViewLongLat LongLatView(LongLatImage, Slice);
-		for (uint32 Face = 0; Face < 6; ++Face)
+
+		// Parallel on the 6 faces :
+		ParallelFor( TEXT("Texture.CubeMipFromLongLat.PF"),6,1, [&](int32 Face)
 		{
 			FImageView2D MipView(*OutMip, Slice * 6 + Face);
 			for (uint32 y = 0; y < Extent; ++y)
@@ -1900,10 +1949,10 @@ void ITextureCompressorModule::GenerateBaseCubeMipFromLongitudeLatitude2D(FImage
 				for (uint32 x = 0; x < Extent; ++x)
 				{
 					FVector DirectionWS = ComputeWSCubeDirectionAtTexelCenter(Face, x, y, InvExtent);
-					MipView.Access(x, y) = LongLatView.LookupLongLat(DirectionWS);
+					MipView.Access(x, y) = LongLatView.LookupLongLatDouble(DirectionWS);
 				}
 			}
-		}
+		});
 	}
 }
 
