TextureCompressorModule : fix CopyTo dest GammaSpace

diff --git a/Engine/Source/Developer/TextureCompressor/Private/TextureCompressorModule.cpp b/Engine/Source/Developer/TextureCompressor/Private/TextureCompressorModule.cpp
--- a/Engine/Source/Developer/TextureCompressor/Private/TextureCompressorModule.cpp
+++ b/Engine/Source/Developer/TextureCompressor/Private/TextureCompressorModule.cpp
@@ -3862,8 +3862,14 @@ private:
 					{
 						// if image is in BGRA8 format leave it, otherwise convert to RGBA32F
 						//  we only support leaving images in source format if they are BGRA8 and require no processing (eg VT tiles)
-						ERawImageFormat::Type DestFormat = Image.Format == ERawImageFormat::BGRA8 ? ERawImageFormat::BGRA8 : ERawImageFormat::RGBA32F;
-						Image.CopyTo(*Mip, DestFormat, Image.GammaSpace);
+						ERawImageFormat::Type DestFormat = Image.Format;
+						EGammaSpace DestGammaSpace = Image.GammaSpace;
+						if ( Image.Format != ERawImageFormat::BGRA8 )
+						{
+							DestFormat = ERawImageFormat::RGBA32F;
+							DestGammaSpace = EGammaSpace::Linear;
+						}
+						Image.CopyTo(*Mip, DestFormat, DestGammaSpace);
 					}
 				}
 			}
