TextureCompressorModule: mip gen filter on cube map faces should clamp, not wrap
has no effect on 2x2 or angular mip gen

diff --git a/Engine/Source/Developer/TextureCompressor/Private/TextureCompressorModule.cpp b/Engine/Source/Developer/TextureCompressor/Private/TextureCompressorModule.cpp
--- a/Engine/Source/Developer/TextureCompressor/Private/TextureCompressorModule.cpp
+++ b/Engine/Source/Developer/TextureCompressor/Private/TextureCompressorModule.cpp
@@ -1149,8 +1149,18 @@ static void GenerateMipBorder(
 // how should be treat lookups outside of the image
 static EMipGenAddressMode ComputeAddressMode(const FImage & Image,const FTextureBuildSettings& Settings)
 {
-	// note: all textures Wrap by default even if their address mode is set to Clamp !?
-	EMipGenAddressMode AddressMode = MGTAM_Wrap;
+	if ( Settings.bPreserveBorder )
+	{
+		return Settings.bBorderColorBlack ? MGTAM_BorderBlack : MGTAM_Clamp;
+	}
+	
+	// conditional on bUseNewMipFilter so old textures are not changed
+	//	really this should be done all the time
+	if ( Settings.bCubemap && Settings.bUseNewMipFilter )
+	{
+		// Cubemaps should Clamp on their faces, never wrap :
+		return MGTAM_Clamp;
+	}
 
 	// Wrap uses AND so requires pow2 sizes ; change to Clamp if nonpow2
 	bool bIsPow2 = FMath::IsPowerOfTwo(Image.SizeX) && FMath::IsPowerOfTwo(Image.SizeY);
@@ -1164,13 +1174,15 @@ static EMipGenAddressMode ComputeAddressMode(const FImage & Image,const FTexture
 
 	if ( ! bIsPow2 )
 	{
-		AddressMode = MGTAM_Clamp;
+		return MGTAM_Clamp;
 	}
 
-	if(Settings.bPreserveBorder)
-	{
-		AddressMode = Settings.bBorderColorBlack ? MGTAM_BorderBlack : MGTAM_Clamp;
-	}
+	// note: all textures Wrap by default even if their address mode is set to Clamp !?
+	EMipGenAddressMode AddressMode = MGTAM_Wrap;
+
+	// @todo Oodle : could do Clamp here if both TextureX and Y address modes are set to Clamp (GetTextureAddressX)
+	//	that would catch the most common case with no other code work
+	//	the more complete solution requires separate AddressMode for X/Y/Z
 
 	return AddressMode;
 }
diff --git a/Engine/Source/Runtime/Engine/Private/TextureDerivedData.cpp b/Engine/Source/Runtime/Engine/Private/TextureDerivedData.cpp
--- a/Engine/Source/Runtime/Engine/Private/TextureDerivedData.cpp
+++ b/Engine/Source/Runtime/Engine/Private/TextureDerivedData.cpp
@@ -258,6 +258,18 @@ static void SerializeForKey(FArchive& Ar, const FTextureBuildSettings& Settings)
 		Ar << TempGuid;
 	}
 
+	if ( Settings.bCubemap && Settings.bUseNewMipFilter )
+	{
+		if ( ( Settings.MipGenSettings >= TMGS_Sharpen0 && Settings.MipGenSettings <= TMGS_Sharpen10 ) ||
+			( Settings.MipGenSettings >= TMGS_Blur1 && Settings.MipGenSettings <= TMGS_Blur5 ) )
+		{
+			// @todo SerializeForKey these can go away whenever we bump the overall ddc key
+			// behavior of mip filter changed so modify the key :
+			TempGuid = FGuid(0xB0420236,0x90064562,0x9C1F10B8,0x2771C31F);
+			Ar << TempGuid;			
+		}
+	}
+
 	if ( Settings.MaxTextureResolution != FTextureBuildSettings::MaxTextureResolutionDefault &&
 		( Settings.MipGenSettings == TMGS_LeaveExistingMips || Settings.bDoScaleMipsForAlphaCoverage ) )
 	{
