Texture MipGen clamp mip filter if texture address mode is not wrap in both x&y
todo for the future to separate x/y mip gen clamp/wrap

diff --git a/Engine/Source/Developer/TextureBuild/Private/TextureBuildFunction.cpp b/Engine/Source/Developer/TextureBuild/Private/TextureBuildFunction.cpp
--- a/Engine/Source/Developer/TextureBuild/Private/TextureBuildFunction.cpp
+++ b/Engine/Source/Developer/TextureBuild/Private/TextureBuildFunction.cpp
@@ -159,6 +159,9 @@ static FTextureBuildSettings ReadBuildSettingsFromCompactBinary(const FCbObjectV
 	BuildSettings.bOodleUsesRDO = Object["bOodleUsesRDO"].AsBool(BuildSettings.bOodleUsesRDO);
 	BuildSettings.OodleRDO = Object["OodleRDO"].AsUInt8(BuildSettings.OodleRDO);
 	ReadCbField(Object["OodleTextureSdkVersion"], BuildSettings.OodleTextureSdkVersion);
+	ReadCbField(Object["TextureAddressModeX"], BuildSettings.TextureAddressModeX);
+	ReadCbField(Object["TextureAddressModeY"], BuildSettings.TextureAddressModeY);
+	ReadCbField(Object["TextureAddressModeZ"], BuildSettings.TextureAddressModeZ);
 
 	return BuildSettings;
 }
diff --git a/Engine/Source/Developer/TextureCompressor/Private/TextureCompressorModule.cpp b/Engine/Source/Developer/TextureCompressor/Private/TextureCompressorModule.cpp
--- a/Engine/Source/Developer/TextureCompressor/Private/TextureCompressorModule.cpp
+++ b/Engine/Source/Developer/TextureCompressor/Private/TextureCompressorModule.cpp
@@ -307,7 +307,7 @@ private:
 	}
 
 	// InOutTable += InTable
-	static void AddFilterTable1D( float *InOutTable, float *InTable, uint32 TableSize )
+	static void AddFilterTable1D( float *InOutTable, const float *InTable, uint32 TableSize )
 	{
 		for(uint32 x = 0; x < TableSize; ++x)
 		{
@@ -1180,9 +1180,34 @@ static EMipGenAddressMode ComputeAddressMode(const FImage & Image,const FTexture
 	// note: all textures Wrap by default even if their address mode is set to Clamp !?
 	EMipGenAddressMode AddressMode = MGTAM_Wrap;
 
-	// @todo Oodle : could do Clamp here if both TextureX and Y address modes are set to Clamp (GetTextureAddressX)
-	//	that would catch the most common case with no other code work
-	//	the more complete solution requires separate AddressMode for X/Y/Z
+	if ( Settings.bUseNewMipFilter )
+	{
+		// in new filters, we do use address mode to change mipgen
+		//	(isolated to new filters solely to avoid changing old content)
+
+		// texture address mode is Wrap,Clamp,Mirror, with Wrap by default
+		// treat Clamp & Mirror the same
+		bool bAddressXWrap = (Settings.TextureAddressModeX == TA_Wrap);
+		bool bAddressYWrap = (Settings.TextureAddressModeY == TA_Wrap);
+
+		if ( !bAddressXWrap && !bAddressYWrap )
+		{
+			AddressMode = MGTAM_Clamp;
+		}
+		else if ( bAddressXWrap && bAddressYWrap )
+		{
+			AddressMode = MGTAM_Wrap;
+		}
+		else
+		{
+			// use Wrap for now to match legacy behavior
+			// ideally we'd wrap one dimension and clamp the other, but that is not yet supported
+			// @todo Oodle: separate wrap/clamp for X&Y ; remove the heavy templating on AddressMode
+			AddressMode = MGTAM_Wrap;
+
+			UE_LOG(LogTextureCompressor, Verbose, TEXT("Not ideal: Heterogeneous XY Wrap/Clamp will generate mips with Wrap on both dimensions") );
+		}
+	}
 
 	return AddressMode;
 }
diff --git a/Engine/Source/Developer/TextureCompressor/Public/TextureCompressorModule.h b/Engine/Source/Developer/TextureCompressor/Public/TextureCompressorModule.h
--- a/Engine/Source/Developer/TextureCompressor/Public/TextureCompressorModule.h
+++ b/Engine/Source/Developer/TextureCompressor/Public/TextureCompressorModule.h
@@ -224,6 +224,11 @@ struct FTextureBuildSettings
 	// ETextureEncodeSpeed, either Final or Fast.
 	uint8 RepresentsEncodeSpeedNoSend;
 
+	// "TextureAddress" enum values : (TA_Wrap default)
+	uint8 TextureAddressModeX = 0;
+	uint8 TextureAddressModeY = 0;
+	uint8 TextureAddressModeZ = 0;
+
 	// If the target format is a tiled format and can leverage reusing the linear encoding, this is not nullptr.
 	const ITextureTiler* Tiler = nullptr;
 
diff --git a/Engine/Source/Runtime/Engine/Classes/Engine/Texture.h b/Engine/Source/Runtime/Engine/Classes/Engine/Texture.h
--- a/Engine/Source/Runtime/Engine/Classes/Engine/Texture.h
+++ b/Engine/Source/Runtime/Engine/Classes/Engine/Texture.h
@@ -53,48 +53,6 @@ struct FPropertyChangedEvent;
 namespace UE::DerivedData { struct FValueId; }
 #endif
 
-UENUM()
-enum TextureFilter
-{
-	TF_Nearest UMETA(DisplayName="Nearest"),
-	TF_Bilinear UMETA(DisplayName="Bi-linear"),
-	TF_Trilinear UMETA(DisplayName="Tri-linear"),
-	/** Use setting from the Texture Group. */
-	TF_Default UMETA(DisplayName="Default (from Texture Group)"),
-	TF_MAX,
-};
-
-UENUM()
-enum TextureAddress
-{
-	TA_Wrap UMETA(DisplayName="Wrap"),
-	TA_Clamp UMETA(DisplayName="Clamp"),
-	TA_Mirror UMETA(DisplayName="Mirror"),
-	TA_MAX,
-};
-
-UENUM()
-enum ETextureMipCount
-{
-	TMC_ResidentMips,
-	TMC_AllMips,
-	TMC_AllMipsBiased,
-	TMC_MAX,
-};
-
-// TextureCompressionQuality is used for ASTC
-UENUM()
-enum ETextureCompressionQuality
-{
-	TCQ_Default = 0		UMETA(DisplayName="Default"),
-	TCQ_Lowest = 1		UMETA(DisplayName="Lowest (ASTC 12x12)"),
-	TCQ_Low = 2			UMETA(DisplayName="Low (ASTC 10x10)"),
-	TCQ_Medium = 3		UMETA(DisplayName="Medium (ASTC 8x8)"),
-	TCQ_High= 4			UMETA(DisplayName="High (ASTC 6x6)"),
-	TCQ_Highest = 5		UMETA(DisplayName="Highest (ASTC 4x4)"),
-	TCQ_MAX,
-};
-
 USTRUCT()
 struct FTextureSourceBlock
 {
diff --git a/Engine/Source/Runtime/Engine/Classes/Engine/TextureDefines.h b/Engine/Source/Runtime/Engine/Classes/Engine/TextureDefines.h
--- a/Engine/Source/Runtime/Engine/Classes/Engine/TextureDefines.h
+++ b/Engine/Source/Runtime/Engine/Classes/Engine/TextureDefines.h
@@ -405,6 +405,49 @@ enum class ETextureChromaticAdaptationMethod : uint8
 };
 
 
+UENUM()
+enum TextureFilter
+{
+	TF_Nearest UMETA(DisplayName="Nearest"),
+	TF_Bilinear UMETA(DisplayName="Bi-linear"),
+	TF_Trilinear UMETA(DisplayName="Tri-linear"),
+	/** Use setting from the Texture Group. */
+	TF_Default UMETA(DisplayName="Default (from Texture Group)"),
+	TF_MAX,
+};
+
+UENUM()
+enum TextureAddress
+{
+	TA_Wrap UMETA(DisplayName="Wrap"),
+	TA_Clamp UMETA(DisplayName="Clamp"),
+	TA_Mirror UMETA(DisplayName="Mirror"),
+	TA_MAX,
+};
+
+UENUM()
+enum ETextureMipCount
+{
+	TMC_ResidentMips,
+	TMC_AllMips,
+	TMC_AllMipsBiased,
+	TMC_MAX,
+};
+
+// TextureCompressionQuality is used for ASTC
+UENUM()
+enum ETextureCompressionQuality
+{
+	TCQ_Default = 0		UMETA(DisplayName="Default"),
+	TCQ_Lowest = 1		UMETA(DisplayName="Lowest (ASTC 12x12)"),
+	TCQ_Low = 2			UMETA(DisplayName="Low (ASTC 10x10)"),
+	TCQ_Medium = 3		UMETA(DisplayName="Medium (ASTC 8x8)"),
+	TCQ_High= 4			UMETA(DisplayName="High (ASTC 6x6)"),
+	TCQ_Highest = 5		UMETA(DisplayName="Highest (ASTC 4x4)"),
+	TCQ_MAX,
+};
+
+
 namespace UE
 {
 namespace TextureDefines
diff --git a/Engine/Source/Runtime/Engine/Private/TextureDerivedData.cpp b/Engine/Source/Runtime/Engine/Private/TextureDerivedData.cpp
--- a/Engine/Source/Runtime/Engine/Private/TextureDerivedData.cpp
+++ b/Engine/Source/Runtime/Engine/Private/TextureDerivedData.cpp
@@ -323,6 +323,16 @@ static void SerializeForKey(FArchive& Ar, const FTextureBuildSettings& Settings)
 		TempGuid = Settings.CompressionCacheId; Ar << TempGuid;
 	}
 
+	if ( Settings.bUseNewMipFilter )
+	{
+		// @todo SerializeForKey : TextureAddressModeX is only used if bUseNewMipFilter is true
+		//	so we hide it in here to avoid changing more DDC keys
+		// todo: when there is an overall DDC key bump, remove ths if on NewFilter so this is just always written
+		TempByte = Settings.TextureAddressModeX; Ar << TempByte;
+		TempByte = Settings.TextureAddressModeY; Ar << TempByte;
+		TempByte = Settings.TextureAddressModeZ; Ar << TempByte;
+	}
+
 	// Note - compression quality is added to the DDC by the formats (based on whether they
 	// use them or not).
 	// This is true for:
@@ -1037,6 +1047,10 @@ static void GetTextureBuildSettings(
 		OutBuildSettings.VirtualTextureTileSize = 0;
 		OutBuildSettings.VirtualTextureBorderSize = 0;
 	}
+	
+	OutBuildSettings.TextureAddressModeX = Texture.GetTextureAddressX();
+	OutBuildSettings.TextureAddressModeY = Texture.GetTextureAddressY();
+	OutBuildSettings.TextureAddressModeZ = Texture.GetTextureAddressZ();
 
 	// By default, initialize settings for layer0
 	FinalizeBuildSettingsForLayer(Texture, 0, &TargetPlatform, InEncodeSpeed, OutBuildSettings, OutBuildResultMetadata);
diff --git a/Engine/Source/Runtime/Engine/Private/TextureDerivedDataBuildUtils.cpp b/Engine/Source/Runtime/Engine/Private/TextureDerivedDataBuildUtils.cpp
--- a/Engine/Source/Runtime/Engine/Private/TextureDerivedDataBuildUtils.cpp
+++ b/Engine/Source/Runtime/Engine/Private/TextureDerivedDataBuildUtils.cpp
@@ -193,7 +193,12 @@ static void WriteBuildSettings(FCbWriter& Writer, const FTextureBuildSettings& B
 	WriteCbFieldWithDefault<bool>(Writer, "bOodleUsesRDO", BuildSettings.bOodleUsesRDO, DefaultSettings.bOodleUsesRDO);
 	
 	WriteCbFieldWithDefault(Writer, "OodleTextureSdkVersion", BuildSettings.OodleTextureSdkVersion, DefaultSettings.OodleTextureSdkVersion);
-	
+
+	WriteCbFieldWithDefault(Writer, "TextureAddressModeX", BuildSettings.TextureAddressModeX, DefaultSettings.TextureAddressModeX);
+	WriteCbFieldWithDefault(Writer, "TextureAddressModeY", BuildSettings.TextureAddressModeY, DefaultSettings.TextureAddressModeY);
+	WriteCbFieldWithDefault(Writer, "TextureAddressModeZ", BuildSettings.TextureAddressModeZ, DefaultSettings.TextureAddressModeZ);
+
+	// @todo SerializeForKey : remove these when overall key changes
 	if ( BuildSettings.bVolume )
 	{
 		WriteCbField<bool>(Writer, "bVolume_ForceNewDDcKey", true); 
