Removing CompressionCacheGuid from texture build settings as it's not needed there.

diff --git a/Engine/Source/Developer/TextureBuild/Private/TextureBuildFunction.cpp b/Engine/Source/Developer/TextureBuild/Private/TextureBuildFunction.cpp
--- a/Engine/Source/Developer/TextureBuild/Private/TextureBuildFunction.cpp
+++ b/Engine/Source/Developer/TextureBuild/Private/TextureBuildFunction.cpp
@@ -84,7 +84,6 @@ static FTextureBuildSettings ReadBuildSettingsFromCompactBinary(const FCbObjectV
 	FTextureBuildSettings BuildSettings;
 	BuildSettings.FormatConfigOverride = Object["FormatConfigOverride"].AsObjectView();
 	FCbObjectView ColorAdjustmentCbObj = Object["ColorAdjustment"].AsObjectView();
-	ReadCbField(ColorAdjustmentCbObj["CompressionCacheId"], BuildSettings.CompressionCacheId);
 	FColorAdjustmentParameters& ColorAdjustment = BuildSettings.ColorAdjustment;
 	ReadCbField(ColorAdjustmentCbObj["AdjustBrightness"], ColorAdjustment.AdjustBrightness);
 	ReadCbField(ColorAdjustmentCbObj["AdjustBrightnessCurve"], ColorAdjustment.AdjustBrightnessCurve);
diff --git a/Engine/Source/Developer/TextureCompressor/Public/TextureCompressorModule.h b/Engine/Source/Developer/TextureCompressor/Public/TextureCompressorModule.h
--- a/Engine/Source/Developer/TextureCompressor/Public/TextureCompressorModule.h
+++ b/Engine/Source/Developer/TextureCompressor/Public/TextureCompressorModule.h
@@ -76,9 +76,11 @@ struct FColorAdjustmentParameters
  */
 struct FTextureBuildSettings
 {
-	/** An optional ID to change the cache key. */
-	FGuid CompressionCacheId;
-	/** Format specific config object view or null if no format specific config is applied as part of this build. */
+	/** Format specific config object view or null if no format specific config is applied as part of this build. This is only for DDC2 builds,
+	* and gets created when the build settings gets serialized for sending - it is not valid beforehand. Meaning, it should only ever be read from
+	* build workers. It is used in place of INIs or command line args to configure individual texture formats since those are not available to
+	* the texture build workers.
+	*/
 	FCbObjectView FormatConfigOverride;
 	/** Color adjustment parameters. */
 	FColorAdjustmentParameters ColorAdjustment;
@@ -114,7 +116,7 @@ struct FTextureBuildSettings
 	uint32 bLongLatSource : 1;
 	/** Whether the texture contains color data in the sRGB colorspace. */
 	uint32 bSRGB : 1;
-	/** Advanced source encoding of the image. */
+	/** Advanced source encoding of the image. UE::Color::EEncoding / ETextureSourceEncoding (same thing) */
 	uint8 SourceEncodingOverride;
 	/** Whether the texture has a defined source color space. */
 	bool bHasColorSpaceDefinition;
@@ -160,7 +162,10 @@ struct FTextureBuildSettings
 	uint8 CompositeTextureMode;	// ECompositeTextureMode, opaque to avoid dependencies on engine headers.
 	/* default 1, high values result in a stronger effect */
 	float CompositePower;
-	/** The source texture's final LOD bias (i.e. includes LODGroup based biases) */
+	/** The source texture's final LOD bias (i.e. includes LODGroup based biases). Generally this does not affect the built texture as the
+	* mips are stripped during cooking, however for tiling some platforms require knowing the actual texture size that will be created. In this
+	* case they need to know the LOD bias to compensate.
+	*/
 	uint32 LODBias;
 	/** The source texture's final LOD bias (i.e. includes LODGroup based biases). This allows cinematic mips as well. */
 	uint32 LODBiasWithCinematicMips;
diff --git a/Engine/Source/Runtime/Engine/Private/TextureDerivedData.cpp b/Engine/Source/Runtime/Engine/Private/TextureDerivedData.cpp
--- a/Engine/Source/Runtime/Engine/Private/TextureDerivedData.cpp
+++ b/Engine/Source/Runtime/Engine/Private/TextureDerivedData.cpp
@@ -318,10 +318,6 @@ static void SerializeForKey(FArchive& Ar, const FTextureBuildSettings& Settings)
 		Ar << TempGuid;
 	}
 
-	if (Settings.CompressionCacheId.IsValid())
-	{
-		TempGuid = Settings.CompressionCacheId; Ar << TempGuid;
-	}
 
 	if ( Settings.bUseNewMipFilter )
 	{
@@ -432,6 +428,12 @@ void GetTextureDerivedDataKeySuffix(const UTexture& Texture, const FTextureBuild
 	FMemoryWriter Ar(TempBytes, /*bIsPersistent=*/ true);
 	SerializeForKey(Ar, BuildSettings);
 
+	if (Texture.CompressionCacheId.IsValid())
+	{
+		FGuid TempGuid = Texture.CompressionCacheId;
+		Ar << TempGuid;
+	}
+
 	for (int32 LayerIndex = 1; LayerIndex < NumLayers; ++LayerIndex)
 	{
 		const FTextureBuildSettings& LayerBuildSettings = BuildSettingsPerLayer[LayerIndex];
@@ -869,7 +871,6 @@ static void GetTextureBuildSettings(
 		OutBuildSettings.AlphaCoverageThresholds = FVector4f(0,0,0,0);
 	}
 
-	OutBuildSettings.CompressionCacheId = Texture.CompressionCacheId;
 	OutBuildSettings.bUseNewMipFilter = Texture.bUseNewMipFilter;
 	OutBuildSettings.bNormalizeNormals = Texture.bNormalizeNormals && Texture.IsNormalMap();
 	OutBuildSettings.bComputeBokehAlpha = (Texture.LODGroup == TEXTUREGROUP_Bokeh);
diff --git a/Engine/Source/Runtime/Engine/Private/TextureDerivedDataBuildUtils.cpp b/Engine/Source/Runtime/Engine/Private/TextureDerivedDataBuildUtils.cpp
--- a/Engine/Source/Runtime/Engine/Private/TextureDerivedDataBuildUtils.cpp
+++ b/Engine/Source/Runtime/Engine/Private/TextureDerivedDataBuildUtils.cpp
@@ -128,7 +128,6 @@ static void WriteBuildSettings(FCbWriter& Writer, const FTextureBuildSettings& B
 		// AlphaCoverageThresholds do not affect build if bDoScaleMipsForAlphaCoverage is off
 		WriteCbFieldWithDefault(Writer, "AlphaCoverageThresholds", BuildSettings.AlphaCoverageThresholds, DefaultSettings.AlphaCoverageThresholds);
 	}
-	WriteCbFieldWithDefault(Writer, "CompressionCacheId", BuildSettings.CompressionCacheId, DefaultSettings.CompressionCacheId);
 	WriteCbFieldWithDefault<bool>(Writer, "bUseNewMipFilter", BuildSettings.bUseNewMipFilter, DefaultSettings.bUseNewMipFilter);
 	WriteCbFieldWithDefault<bool>(Writer, "bNormalizeNormals", BuildSettings.bNormalizeNormals, DefaultSettings.bNormalizeNormals);
 	WriteCbFieldWithDefault(Writer, "MipSharpening", BuildSettings.MipSharpening, DefaultSettings.MipSharpening);
@@ -310,6 +309,12 @@ FCbObject SaveTextureBuildSettings(const UTexture& Texture, const FTextureBuildS
 
 	Writer.AddUuid("BuildVersion", GetTextureDerivedDataVersion());
 	
+	if (Texture.CompressionCacheId.IsValid())
+	{
+		// Not actually read by the worker - just used to make a different key
+		Writer.AddUuid("CompressionCacheId", Texture.CompressionCacheId);
+	}
+
 	Writer.AddInteger("RequiredMemoryEstimate", RequiredMemoryEstimate);
 
 	if (uint16 TextureFormatVersion = TextureFormat->GetVersion(BuildSettings.TextureFormatName, &BuildSettings))
