TextureCompressor: Disable cplusplus.NewDelete static analysis checker

#rnx

diff --git a/Engine/Source/Developer/TextureCompressor/TextureCompressor.Build.cs b/Engine/Source/Developer/TextureCompressor/TextureCompressor.Build.cs
--- a/Engine/Source/Developer/TextureCompressor/TextureCompressor.Build.cs
+++ b/Engine/Source/Developer/TextureCompressor/TextureCompressor.Build.cs
@@ -22,5 +22,8 @@ public class TextureCompressor : ModuleRules
 			);
 
 		//AddEngineThirdPartyPrivateStaticDependencies(Target, "nvTextureTools");
+
+		// TODO: TextureCompressorModule.cpp: warning: Use of memory after it is freed (within a call to 'GetResult') [cplusplus.NewDelete]
+		StaticAnalyzerDisabledCheckers.Add("cplusplus.NewDelete");
 	}
 }
