Tasks: Suppress clang static analysis warning when deleting refcounted FTaskBase

#rnx

diff --git a/Engine/Source/Developer/TextureCompressor/TextureCompressor.Build.cs b/Engine/Source/Developer/TextureCompressor/TextureCompressor.Build.cs
--- a/Engine/Source/Developer/TextureCompressor/TextureCompressor.Build.cs
+++ b/Engine/Source/Developer/TextureCompressor/TextureCompressor.Build.cs
@@ -22,8 +22,5 @@ public class TextureCompressor : ModuleRules
 			);
 
 		//AddEngineThirdPartyPrivateStaticDependencies(Target, "nvTextureTools");
-
-		// TODO: TextureCompressorModule.cpp: warning: Use of memory after it is freed (within a call to 'GetResult') [cplusplus.NewDelete]
-		StaticAnalyzerDisabledCheckers.Add("cplusplus.NewDelete");
 	}
 }
diff --git a/Engine/Source/Runtime/Core/Public/Tasks/TaskPrivate.h b/Engine/Source/Runtime/Core/Public/Tasks/TaskPrivate.h
--- a/Engine/Source/Runtime/Core/Public/Tasks/TaskPrivate.h
+++ b/Engine/Source/Runtime/Core/Public/Tasks/TaskPrivate.h
@@ -114,7 +114,9 @@ namespace UE::Tasks
 				uint32 LocalRefCount = RefCount.fetch_sub(1, std::memory_order_acq_rel) - 1;
 				if (LocalRefCount == 0)
 				{
+#if !defined(__clang_analyzer__)
 					delete this;
+#endif
 				}
 			}
 
