TextureCompressorModule : remove unnecessary and inconsistent conditional on ParallelFor in AdustImage

diff --git a/Engine/Source/Developer/TextureCompressor/Private/TextureCompressorModule.cpp b/Engine/Source/Developer/TextureCompressor/Private/TextureCompressorModule.cpp
--- a/Engine/Source/Developer/TextureCompressor/Private/TextureCompressorModule.cpp
+++ b/Engine/Source/Developer/TextureCompressor/Private/TextureCompressorModule.cpp
@@ -445,7 +445,7 @@ static FVector4f ComputeAlphaCoverage(const FVector4f Thresholds, const FVector4
 
 		const float ThresholdScaled = DetermineScaledThreshold(Thresholds[3] , Scales[3]);
 		
-		int32 CommonResult = 0;
+		int64 CommonResult = 0;
 		ParallelFor( TEXT("Texture.ComputeAlphaCoverage.PF"),NumJobs,1, [&](int32 Index)
 		{
 			int32 StartIndex = Index * NumRowsEachJob;
@@ -464,7 +464,7 @@ static FVector4f ComputeAlphaCoverage(const FVector4f Thresholds, const FVector4
 			FPlatformAtomics::InterlockedAdd(&CommonResult, LocalCoverage);
 		});
 
-		Coverage[3] = float(CommonResult) / float(SourceImageData.SizeX * SourceImageData.SizeY);
+		Coverage[3] = float(CommonResult) / float( (int64) SourceImageData.SizeX * SourceImageData.SizeY);
 		
 		UE_LOG(LogTextureCompressor, VeryVerbose, TEXT("Thresholds = 000 %f Coverage = 000 %f"),  \
 			Thresholds[3], Coverage[3] );
@@ -488,7 +488,7 @@ static FVector4f ComputeAlphaCoverage(const FVector4f Thresholds, const FVector4
 			}
 		}
 
-		int32 CommonResults[4] = { 0, 0, 0, 0 };
+		int64 CommonResults[4] = { 0, 0, 0, 0 };
 		ParallelFor( TEXT("Texture.ComputeAlphaCoverage.PF"),NumJobs,1, [&](int32 Index)
 		{
 			int32 StartIndex = Index * NumRowsEachJob;
@@ -518,7 +518,7 @@ static FVector4f ComputeAlphaCoverage(const FVector4f Thresholds, const FVector4
 
 		for (int32 i = 0; i < 4; ++i)
 		{
-			Coverage[i] = float(CommonResults[i]) / float(SourceImageData.SizeX * SourceImageData.SizeY);
+			Coverage[i] = float(CommonResults[i]) / float((int64) SourceImageData.SizeX * SourceImageData.SizeY);
 		}
 		
 		UE_LOG(LogTextureCompressor, VeryVerbose, TEXT("Thresholds = %f %f %f %f Coverage = %f %f %f %f"),  \
@@ -2698,12 +2698,6 @@ void ITextureCompressorModule::AdjustImageColors(FImage& Image, const FTextureBu
 		int64 NumPixelsEachJob;
 		int32 NumJobs = ImageParallelForComputeNumJobsForPixels(NumPixelsEachJob,NumPixels);
 
-		// bForceSingleThread is set to true when: 
-		// editor or cooker is loading as this is when the derived data cache is rebuilt as it will already be limited to a single thread 
-		//     and thus overhead of multithreading will simply make it slower
-		// @todo Oodle - this is done here and not in other similar ParallelFor places here.  It should either be done everywhere or nowhere.
-		bool bForceSingleThread = GIsEditorLoadingPackage || GIsCookerLoadingPackage || IsInAsyncLoadingThread();
-
 		ParallelFor( TEXT("Texture.AdjustImageColorsFunc.PF"),NumJobs,1, [&](int32 Index)
 		{
 			int64 StartIndex = Index * NumPixelsEachJob;
@@ -2722,8 +2716,7 @@ void ITextureCompressorModule::AdjustImageColors(FImage& Image, const FTextureBu
 			{
 				AdjustColorsOld(First, Count, InBuildSettings);
 			}
-		}
-		, (bForceSingleThread ? EParallelForFlags::ForceSingleThread : EParallelForFlags::None) );
+		} );
 	}
 }
 
