Fix for longlat textures and LeaveExistingMips calculating the wrong mip count.

diff --git a/Engine/Source/Developer/TextureCompressor/Private/TextureCompressorModule.cpp b/Engine/Source/Developer/TextureCompressor/Private/TextureCompressorModule.cpp
--- a/Engine/Source/Developer/TextureCompressor/Private/TextureCompressorModule.cpp
+++ b/Engine/Source/Developer/TextureCompressor/Private/TextureCompressorModule.cpp
@@ -3343,6 +3343,16 @@ public:
 	{
 		if (BuildSettings.MipGenSettings == TMGS_LeaveExistingMips)
 		{
+			if (BuildSettings.bLongLatSource)
+			{
+				// Longlat sources with leave existing are required to only have 1 mip source (BuildTextureMips checks() on this condition).
+				uint32 LongLatCubemapExtents = ComputeLongLatCubemapExtents(InMip0SizeX, BuildSettings.MaxTextureResolution);
+				OutMip0SizeX = LongLatCubemapExtents;
+				OutMip0SizeY = LongLatCubemapExtents;
+				OutMip0NumSlices = 6 * InMip0NumSlices;
+				return 1;
+			}
+
 			// Since we can't generate, we only have to limit to MaxTextureSize
 			OutMip0NumSlices = InMip0NumSlices; // At the moment, when importing a volume texture, only the 2d dimensions are checked against max resolution.
 			return GetMipCountForLeaveExistingMips(InMip0SizeX, InMip0SizeY, InExistingMipCount, BuildSettings.MaxTextureResolution, OutMip0SizeX, OutMip0SizeY);
