Imported partial mip chains must be generated out to a full mip chain.

diff --git a/Engine/Plugins/Interchange/Runtime/Source/Import/Private/Texture/InterchangeDDSTranslator.cpp b/Engine/Plugins/Interchange/Runtime/Source/Import/Private/Texture/InterchangeDDSTranslator.cpp
--- a/Engine/Plugins/Interchange/Runtime/Source/Import/Private/Texture/InterchangeDDSTranslator.cpp
+++ b/Engine/Plugins/Interchange/Runtime/Source/Import/Private/Texture/InterchangeDDSTranslator.cpp
@@ -94,6 +94,15 @@ namespace UE::Interchange::Private::InterchangeDDSTranslator
 			return {};
 		}
 
+		// We require a complete mip chain or a single mip.
+		if (DDS->MipCount > 1)
+		{
+			if (DDS->MipCount != FImageCoreUtils::GetMipCountFromDimensions(DDS->Width, DDS->Height, DDS->Depth, DDS->IsValidTextureVolume()))
+			{
+				UE_LOG(LogInterchangeImport, Warning, TEXT("Imported DDS files must either be a single base mip, or a full mip chain, not a partial mip chain (found %d expected %d)"), DDS->MipCount, FImageCoreUtils::GetMipCountFromDimensions(DDS->Width, DDS->Height, DDS->Depth, DDS->IsValidTextureVolume()));
+			}
+		}
+
 		if (!bHeaderOnly)
 		{
 			// change X8 formats to A8 :	
diff --git a/Engine/Source/Developer/TextureCompressor/Private/TextureCompressorModule.cpp b/Engine/Source/Developer/TextureCompressor/Private/TextureCompressorModule.cpp
--- a/Engine/Source/Developer/TextureCompressor/Private/TextureCompressorModule.cpp
+++ b/Engine/Source/Developer/TextureCompressor/Private/TextureCompressorModule.cpp
@@ -4,6 +4,7 @@
 #include "Math/RandomStream.h"
 #include "ChildTextureFormat.h"
 #include "Containers/IndirectArray.h"
+#include "ImageCoreUtils.h"
 #include "Stats/Stats.h"
 #include "Async/AsyncWork.h"
 #include "Async/ParallelFor.h"
@@ -3253,30 +3254,6 @@ static void NormalizeMip(FImage& InOutMip)
 }
 
 
-// Special case for TMGS_LeaveExistingMips
-static int32 GetMipCountForLeaveExistingMips(int32 InMip0SizeX, int32 InMip0SizeY, int32 InExistingMipCount, uint32 InMaxTexture2DResolution, int32& OutMip0SizeX, int32& OutMip0SizeY)
-{
-	int32 i = 0;
-	for (; i < InExistingMipCount; i++)
-	{
-		uint32 MipSizeX = FMath::Max<uint32>(1, InMip0SizeX >> i);
-		uint32 MipSizeY = FMath::Max<uint32>(1, InMip0SizeY >> i);
-
-		if (MipSizeX <= InMaxTexture2DResolution &&
-			MipSizeY <= InMaxTexture2DResolution)
-		{
-			OutMip0SizeX = MipSizeX;
-			OutMip0SizeY = MipSizeY;
-			return InExistingMipCount - i;
-		}
-	}
-
-	// Couldn't find a fit, texture build will fail.
-	check(0);
-	return 0;
-}
-
-
 // Returns true if the target texture size is different and padding/stretching is required.
 static bool GetPowerOfTwoTargetTextureSize(int32 InMip0SizeX, int32 InMip0SizeY, int32 InMip0NumSlices, bool bInIsVolume, ETexturePowerOfTwoSetting::Type InPow2Setting, int32& OutTargetSizeX, int32& OutTargetSizeY, int32& OutTargetSizeZ)
 {
@@ -3343,55 +3320,36 @@ public:
 		const FTextureBuildSettings& BuildSettings, 
 		int32& OutMip0SizeX, int32& OutMip0SizeY, int32& OutMip0NumSlices) const override
 	{
-		if (BuildSettings.MipGenSettings == TMGS_LeaveExistingMips)
-		{
-			if (BuildSettings.bLongLatSource)
-			{
-				// Longlat sources with leave existing are required to only have 1 mip source (BuildTextureMips checks() on this condition).
-				uint32 LongLatCubemapExtents = ComputeLongLatCubemapExtents(InMip0SizeX, BuildSettings.MaxTextureResolution);
-				OutMip0SizeX = LongLatCubemapExtents;
-				OutMip0SizeY = LongLatCubemapExtents;
-				OutMip0NumSlices = 6 * InMip0NumSlices;
-				return 1;
-			}
-
-			// Since we can't generate, we only have to limit to MaxTextureSize
-			OutMip0NumSlices = InMip0NumSlices; // At the moment, when importing a volume texture, only the 2d dimensions are checked against max resolution.
-			return GetMipCountForLeaveExistingMips(InMip0SizeX, InMip0SizeY, InExistingMipCount, BuildSettings.MaxTextureResolution, OutMip0SizeX, OutMip0SizeY);
-		}
-
 		// AFAICT LatLongCubeMaps don't do any of this - pow2 is broken with them but it runs, and max texture stuff
 		// is handled internally in the extents function.
-
 		int32 BaseSizeX = InMip0SizeX;
 		int32 BaseSizeY = InMip0SizeY;
 		int32 BaseSizeZ = BuildSettings.bVolume ? InMip0NumSlices : 1; // Volume textures are the only type that mip their Z, arrays and cubes are fixed.
 
-		ETexturePowerOfTwoSetting::Type PowerOfTwoMode = (ETexturePowerOfTwoSetting::Type)BuildSettings.PowerOfTwoMode;
-		if (PowerOfTwoMode != ETexturePowerOfTwoSetting::None)
+		// LatLong sources are clamped in ComputeLongLatCubemapExtents
+		if (BuildSettings.bLongLatSource == false)
 		{
-			int32 TargetSizeX, TargetSizeY, TargetSizeZ;
-			bool NeedsAdjustment = GetPowerOfTwoTargetTextureSize(BaseSizeX, BaseSizeY, BaseSizeY, BuildSettings.bVolume, PowerOfTwoMode, TargetSizeX, TargetSizeY, TargetSizeZ);
-			if (NeedsAdjustment)
+			ETexturePowerOfTwoSetting::Type PowerOfTwoMode = (ETexturePowerOfTwoSetting::Type)BuildSettings.PowerOfTwoMode;
+			if (BuildSettings.MipGenSettings != TMGS_LeaveExistingMips &&
+				PowerOfTwoMode != ETexturePowerOfTwoSetting::None)
 			{
-				// In this case we are regenerating the entire mip chain.
-				InExistingMipCount = 1;
-				BaseSizeX = TargetSizeX;
-				BaseSizeY = TargetSizeY;
-				BaseSizeZ = TargetSizeZ; // volume textures already accounted for
+				int32 TargetSizeX, TargetSizeY, TargetSizeZ;
+				bool NeedsAdjustment = GetPowerOfTwoTargetTextureSize(BaseSizeX, BaseSizeY, BaseSizeY, BuildSettings.bVolume, PowerOfTwoMode, TargetSizeX, TargetSizeY, TargetSizeZ);
+				if (NeedsAdjustment)
+				{
+					// In this case we are regenerating the entire mip chain.
+					InExistingMipCount = 1;
+					BaseSizeX = TargetSizeX;
+					BaseSizeY = TargetSizeY;
+					BaseSizeZ = TargetSizeZ; // volume textures already accounted for
+				}
+				// Otherwise we have valid pow2 so we can reuse any existing mips and regenerate
+				// any missing tail mips.
 			}
-			// Otherwise we have valid pow2 so we can reuse any existing mips and regenerate
-			// any missing tail mips.
-		}
 
-		// LatLong sources are clamped in ComputeLongLatCubemapExtents
-		if (BuildSettings.bLongLatSource == false)
-		{
 			// Max texture resolution strips off mips that are above the limit.
 			int64 MaxTextureResolution = BuildSettings.MaxTextureResolution;
-
-			uint32 GeneratedMaxMipDimension = FMath::Max3(BaseSizeX, BaseSizeY, BaseSizeZ);
-			int32 GeneratedMipCount = 1 + FMath::FloorLog2(GeneratedMaxMipDimension);
+			int32 GeneratedMipCount = FImageCoreUtils::GetMipCountFromDimensions(BaseSizeX, BaseSizeY, BaseSizeZ, BuildSettings.bVolume);
 			int32 i = 0;
 			for (; i < GeneratedMipCount; i++)
 			{
@@ -3454,10 +3412,16 @@ public:
 			return 1;
 		}
 
+		// LeaveExisting is often unintentionally used as "NoMipmaps", so if they bring in 1 mip, leave it as 1 mip.
+		if (BuildSettings.MipGenSettings == TMGS_LeaveExistingMips &&
+			InExistingMipCount == 1)
+		{
+			return 1;
+		}
+
 		// NumOutputMips is the number of mips that would be made if you made a full mip chain
 		//  eg. 256 makes 9 mips , 300 also makes 9 mips
-		uint32 MaxMipDimension = FMath::Max3(BaseSizeX, BaseSizeY, BaseSizeZ);
-		return 1 + FMath::FloorLog2(MaxMipDimension);
+		return FImageCoreUtils::GetMipCountFromDimensions(BaseSizeX, BaseSizeY, BaseSizeZ, BuildSettings.bVolume);
 	}
 
 	virtual bool BuildTexture(
@@ -3757,14 +3721,6 @@ private:
 
 			if ( StartMip == NumSourceMips )
 			{
-				if (BuildSettings.MipGenSettings == TMGS_LeaveExistingMips)
-				{
-					UE_LOG(LogTextureCompressor, Error,	TEXT("LeaveExistingMips no mip that fits max dimension (%d)."),(int)MaxTextureResolution);
-					return false;
-				}
-
-				// currently only makes mips if you only had 1 source mip :
-				check(NumSourceMips == 1);
 
 				// bLongLatCubemap should not get here because cube size is made from MaxTextureSize
 				check( ! bLongLatCubemap );
@@ -3818,7 +3774,7 @@ private:
 				// [StartMip] will now references BuildSourceImageMips.Last()
 			}
 		}
-				
+
 		// now shrinking to MaxTextureResolution is done, figure out which mips to use or make
 		
 		// Copy over base mip and any LeaveExisting, from SourceMips , starting at StartMip
@@ -3832,28 +3788,35 @@ private:
 			NumOutputMips = 1;
 			CopyCount = 1;
 		}
-		else if ( BuildSettings.MipGenSettings == TMGS_LeaveExistingMips )
-		{
-			// only output what we can copy, generate none
-			NumOutputMips = CopyCount;
-		}
 		else
 		{
 			const FImage & TopMip = (*pSourceMips)[StartMip];
 
-			int32 TopMipSizeZ = BuildSettings.bVolume ? TopMip.NumSlices : 1;
-		
 			// NumOutputMips is the number of mips that would be made if you made a full mip chain
 			//  eg. 256 makes 9 mips , 300 also makes 9 mips
-			NumOutputMips = 1 + FMath::FloorLog2(
-				bLongLatCubemap ?
-				ComputeLongLatCubemapExtents(TopMip.SizeX, BuildSettings.MaxTextureResolution) :
-				FMath::Max3(TopMip.SizeX, TopMip.SizeY, TopMipSizeZ) );
-				
-			// unless LeaveExistingMips, we only copy 1 
-			//  (in theory we could copy some existing and generate the rest, but that's not done currently)
-			//  (intentionally so, artists use this to limit mipping down)
-			CopyCount = 1;
+			if (bLongLatCubemap)
+			{
+				NumOutputMips = FImageCoreUtils::GetMipCountFromDimensions(ComputeLongLatCubemapExtents(TopMip.SizeX, BuildSettings.MaxTextureResolution), 1, 1, false);
+			}
+			else
+			{
+				NumOutputMips = FImageCoreUtils::GetMipCountFromDimensions(TopMip.SizeX, TopMip.SizeY, TopMip.NumSlices, BuildSettings.bVolume);
+			}
+
+			// LeaveExistingMips is often inadvertently used as "NoMips", so if they only brought in 1 mip, leave it as
+			// 1 mip.
+			if (BuildSettings.MipGenSettings == TMGS_LeaveExistingMips &&
+				InSourceMipChain.Num() == 1)
+			{
+				NumOutputMips = 1;
+				check(CopyCount == 1);
+			}
+
+			// unless LeaveExistingMips, we only copy 1
+			if (BuildSettings.MipGenSettings != TMGS_LeaveExistingMips)
+			{
+				CopyCount = 1;
+			}
 		}
 
 		// we will output NumOutputMips
diff --git a/Engine/Source/Runtime/ImageCore/Private/ImageCoreUtils.cpp b/Engine/Source/Runtime/ImageCore/Private/ImageCoreUtils.cpp
--- a/Engine/Source/Runtime/ImageCore/Private/ImageCoreUtils.cpp
+++ b/Engine/Source/Runtime/ImageCore/Private/ImageCoreUtils.cpp
@@ -4,8 +4,14 @@
 
 DEFINE_LOG_CATEGORY_STATIC(LogImageCoreUtils, Log, All);
 
+IMAGECORE_API int32 FImageCoreUtils::GetMipCountFromDimensions(int32 InSizeX, int32 InSizeY, int32 InVolumeZ, bool bInIsVolume)
+{
+	uint32 MaxMipDimension = FMath::Max3(InSizeX, InSizeY, bInIsVolume ? InVolumeZ : 1);
+	return 1 + FMath::FloorLog2(MaxMipDimension);
+}
 
-EPixelFormat FImageCoreUtils::GetPixelFormatForRawImageFormat(ERawImageFormat::Type InFormat, 
+
+IMAGECORE_API EPixelFormat FImageCoreUtils::GetPixelFormatForRawImageFormat(ERawImageFormat::Type InFormat, 
 							ERawImageFormat::Type * pOutEquivalentFormat)
 {
 	ERawImageFormat::Type OutEquivalentFormatDummy;
@@ -165,4 +171,4 @@ IMAGECORE_API ETextureSourceFormat FImageCoreUtils::GetCommonSourceFormat(ETextu
 	default:
 		return TSF_RGBA32F;
 	}
-}
+}
\ No newline at end of file
diff --git a/Engine/Source/Runtime/ImageCore/Public/ImageCoreUtils.h b/Engine/Source/Runtime/ImageCore/Public/ImageCoreUtils.h
--- a/Engine/Source/Runtime/ImageCore/Public/ImageCoreUtils.h
+++ b/Engine/Source/Runtime/ImageCore/Public/ImageCoreUtils.h
@@ -20,6 +20,9 @@ ImageCoreUtils.h: Image utility functions.
 
 namespace FImageCoreUtils
 {
+	// Returns the number of mips that constitute a full mip chain for the given top level mip size.
+	// InVolumeZ is ignored unless bInIsVolume is true
+	IMAGECORE_API int32 GetMipCountFromDimensions(int32 InSizeX, int32 InSizeY, int32 InVolumeZ, bool bInIsVolume);
 
 	// ETextureSourceFormat and ERawImageFormat::Type are one-to-one :
 	IMAGECORE_API ERawImageFormat::Type ConvertToRawImageFormat(ETextureSourceFormat Format);
