fix crash in Texture when using Downscale on very skew textures

diff --git a/Engine/Source/Developer/TextureCompressor/Private/TextureCompressorModule.cpp b/Engine/Source/Developer/TextureCompressor/Private/TextureCompressorModule.cpp
--- a/Engine/Source/Developer/TextureCompressor/Private/TextureCompressorModule.cpp
+++ b/Engine/Source/Developer/TextureCompressor/Private/TextureCompressorModule.cpp
@@ -1470,13 +1470,13 @@ static void DownscaleImage(const FImage& SrcImage, FImage& DstImage, const FText
 	AvgKernel.BuildSeparatableGaussWithSharpen(2);
 
 	TArray<FLinearColor> TempData;
-	AllocateTempForMips(TempData, SrcImage.SizeX, SrcImage.SizeY, SrcImage.SizeX / 2, SrcImage.SizeY / 2, false, AvgKernel, 2, false, bUnfiltered, bUseNewMipFilter);
+	AllocateTempForMips(TempData, SrcImage.SizeX, SrcImage.SizeY, FMath::Max(1, SrcImage.SizeX / 2), FMath::Max(1, SrcImage.SizeY / 2), false, AvgKernel, 2, false, bUnfiltered, bUseNewMipFilter);
 
 	int32 NumIterations = 0;
 	while(Downscale > 2.0f)
 	{
-		int32 DstSizeX = ImageChain[0]->SizeX / 2;
-		int32 DstSizeY = ImageChain[0]->SizeY / 2;
+		int32 DstSizeX = FMath::Max(1, ImageChain[0]->SizeX / 2 );
+		int32 DstSizeY = FMath::Max(1, ImageChain[0]->SizeY / 2 );
 		ImageChain[1]->Init(DstSizeX, DstSizeY, ImageChain[0]->NumSlices, ImageChain[0]->Format, ImageChain[0]->GammaSpace);
 
 		FImageView2D SrcImageData(*ImageChain[0], 0);
