TextureCompressorModule : fix some int32 overflows

diff --git a/Engine/Source/Developer/TextureCompressor/Private/TextureCompressorModule.cpp b/Engine/Source/Developer/TextureCompressor/Private/TextureCompressorModule.cpp
--- a/Engine/Source/Developer/TextureCompressor/Private/TextureCompressorModule.cpp
+++ b/Engine/Source/Developer/TextureCompressor/Private/TextureCompressorModule.cpp
@@ -45,9 +45,9 @@ struct FImageView2D
 	/** Pointer to colors in the slice. */
 	FLinearColor* SliceColors;
 	/** Width of the slice. */
-	int32 SizeX;
+	int64 SizeX;
 	/** Height of the slice. */
-	int32 SizeY;
+	int64 SizeY;
 
 	FImageView2D() : SliceColors(nullptr), SizeX(0), SizeY(0) {}
 
@@ -1349,10 +1349,10 @@ static float LookupFloatBilinear(const float * FloatPlane,int32 SizeX,int32 Size
 	int32 IntX1 = FMath::Min(IntX0+1, SizeX-1);
 	int32 IntY1 = FMath::Min(IntY0+1, SizeY-1);
 	
-	float Sample00 = FloatPlane[ IntX0 + IntY0 * SizeX ];
-	float Sample10 = FloatPlane[ IntX1 + IntY0 * SizeX ];
-	float Sample01 = FloatPlane[ IntX0 + IntY1 * SizeX ];
-	float Sample11 = FloatPlane[ IntX1 + IntY1 * SizeX ];
+	float Sample00 = FloatPlane[ IntX0 + IntY0 * (int64) SizeX ];
+	float Sample10 = FloatPlane[ IntX1 + IntY0 * (int64) SizeX ];
+	float Sample01 = FloatPlane[ IntX0 + IntY1 * (int64) SizeX ];
+	float Sample11 = FloatPlane[ IntX1 + IntY1 * (int64) SizeX ];
 	float Sample0 = FMath::Lerp(Sample00, Sample10, FractX);
 	float Sample1 = FMath::Lerp(Sample01, Sample11, FractX);
 		
@@ -1811,16 +1811,16 @@ struct FImageViewLongLat
 	/** Image colors. */
 	FLinearColor* ImageColors;
 	/** Width of the image. */
-	int32 SizeX;
+	int64 SizeX;
 	/** Height of the image. */
-	int32 SizeY;
+	int64 SizeY;
 
 	/** Initialization constructor. */
 	explicit FImageViewLongLat(FImage& Image, int32 SliceIndex)
 	{
 		SizeX = Image.SizeX;
 		SizeY = Image.SizeY;
-		ImageColors = (&Image.AsRGBA32F()[0]) + SliceIndex * (int64) SizeY * SizeX;
+		ImageColors = (&Image.AsRGBA32F()[0]) + SliceIndex * SizeY * SizeX;
 	}
 
 	/** Const access to a texel. */
@@ -2194,7 +2194,7 @@ static void GenerateAngularFilteredMip(FImage* DestMip, FImage& SrcMip, float Co
 	{
 		for(int32 x = 0; x < SrcMip.SizeX; ++x)
 		{
-			TexelAreaArray[x + y * SrcMip.SizeX] = ComputeTexelArea(x, y, MipInvSideExtent * 2);
+			TexelAreaArray[x + y * (int64) SrcMip.SizeX] = ComputeTexelArea(x, y, MipInvSideExtent * 2);
 		}
 	}
 
@@ -2732,14 +2732,14 @@ static void ComputeBokehAlpha(FImage& Image)
 {
 	check( Image.SizeX > 0 && Image.SizeY > 0 );
 
-	const int32 NumPixels = Image.SizeX * Image.SizeY * Image.NumSlices;
+	const int64 NumPixels = Image.GetNumPixels();
 	TArrayView64<FLinearColor> ImageColors = Image.AsRGBA32F();
 
 	// compute LinearAverage
 	FLinearColor LinearAverage;
 	{
 		FLinearColor LinearSum(0, 0, 0, 0);
-		for( int32 CurPixelIndex = 0; CurPixelIndex < NumPixels; ++CurPixelIndex )
+		for( int64 CurPixelIndex = 0; CurPixelIndex < NumPixels; ++CurPixelIndex )
 		{
 			LinearSum += ImageColors[ CurPixelIndex ];
 		}
@@ -2761,7 +2761,7 @@ static void ComputeBokehAlpha(FImage& Image)
 	}
 
 	{
-		for( int32 CurPixelIndex = 0; CurPixelIndex < NumPixels; ++CurPixelIndex )
+		for( int64 CurPixelIndex = 0; CurPixelIndex < NumPixels; ++CurPixelIndex )
 		{
 			const FLinearColor OriginalColor = ImageColors[ CurPixelIndex ];
 
@@ -2784,7 +2784,7 @@ static void ReplicateRedChannel( TArray<FImage>& InOutMipChain )
 	{
 		FImage& SrcMip = InOutMipChain[MipIndex];
 		FLinearColor* FirstColor = (&SrcMip.AsRGBA32F()[0]);
-		FLinearColor* LastColor = FirstColor + (SrcMip.SizeX * SrcMip.SizeY * SrcMip.NumSlices);
+		FLinearColor* LastColor = FirstColor + SrcMip.GetNumPixels();
 		for ( FLinearColor* Color = FirstColor; Color < LastColor; ++Color )
 		{
 			*Color = FLinearColor( Color->R, Color->R, Color->R, Color->R );
@@ -2802,7 +2802,7 @@ static void ReplicateAlphaChannel( TArray<FImage>& InOutMipChain )
 	{
 		FImage& SrcMip = InOutMipChain[MipIndex];
 		FLinearColor* FirstColor = (&SrcMip.AsRGBA32F()[0]);
-		FLinearColor* LastColor = FirstColor + (SrcMip.SizeX * SrcMip.SizeY * SrcMip.NumSlices);
+		FLinearColor* LastColor = FirstColor + SrcMip.GetNumPixels();
 		for ( FLinearColor* Color = FirstColor; Color < LastColor; ++Color )
 		{
 			*Color = FLinearColor( Color->A, Color->A, Color->A, Color->A );
@@ -2840,11 +2840,11 @@ static void ApplyYCoCgBlockScale(TArray<FImage>& InOutMipChain)
 
 		for (int32 Slice = 0; Slice < SrcMip.NumSlices; ++Slice)
 		{
-			FLinearColor* SliceFirstColor = FirstColor + (SrcMip.SizeX * SrcMip.SizeY * Slice);
+			FLinearColor* SliceFirstColor = FirstColor + SrcMip.GetSliceNumPixels() * Slice;
 
 			for (int32 Y = 0; Y < BlockWidthY; ++Y)
 			{
-				FLinearColor* RowFirstColor = SliceFirstColor + (Y * 4 * SrcMip.SizeY);
+				FLinearColor* RowFirstColor = SliceFirstColor + Y * 4 * (int64) SrcMip.SizeX;
 
 				for (int32 X = 0; X < BlockWidthX; ++X)
 				{
@@ -2854,7 +2854,7 @@ static void ApplyYCoCgBlockScale(TArray<FImage>& InOutMipChain)
 					float MaxComponent = 0.f;
 					for (int32 BlockY = 0; BlockY < 4; ++BlockY)
 					{
-						FLinearColor* Color = BlockFirstColor + (BlockY * SrcMip.SizeY);
+						FLinearColor* Color = BlockFirstColor + BlockY * SrcMip.SizeX;
 						for (int32 BlockX = 0; BlockX < 4; ++BlockX, ++Color)
 						{
 							MaxComponent = FMath::Max(FMath::Abs(Color->R - 128.f / 255.f), MaxComponent);
@@ -2868,7 +2868,7 @@ static void ApplyYCoCgBlockScale(TArray<FImage>& InOutMipChain)
 					// Iterate block to modify for scale
 					for (int32 BlockY = 0; BlockY < 4; ++BlockY)
 					{
-						FLinearColor* Color = BlockFirstColor + (BlockY * SrcMip.SizeY);
+						FLinearColor* Color = BlockFirstColor + BlockY * SrcMip.SizeX;
 						for (int32 BlockX = 0; BlockX < 4; ++BlockX, ++Color)
 						{
 							const float OutR = (Color->R - 128.f / 255.f) * Scale + 128.f / 255.f;
@@ -3240,9 +3240,9 @@ static void NormalizeMip(FImage& InOutMip)
 {
 	const FVector NormalIfZero(0.f,0.f,1.f);
 
-	const uint32 NumPixels = InOutMip.SizeX * InOutMip.SizeY * InOutMip.NumSlices;
+	const int64 NumPixels = InOutMip.GetNumPixels();
 	TArrayView64<FLinearColor> ImageColors = InOutMip.AsRGBA32F();
-	for(uint32 CurPixelIndex = 0; CurPixelIndex < NumPixels; ++CurPixelIndex)
+	for(int64 CurPixelIndex = 0; CurPixelIndex < NumPixels; ++CurPixelIndex)
 	{
 		FLinearColor& Color = ImageColors[CurPixelIndex];
 
