TextureCompressorModule : Parallel NormalizeMipChain

diff --git a/Engine/Source/Developer/TextureCompressor/Private/TextureCompressorModule.cpp b/Engine/Source/Developer/TextureCompressor/Private/TextureCompressorModule.cpp
--- a/Engine/Source/Developer/TextureCompressor/Private/TextureCompressorModule.cpp
+++ b/Engine/Source/Developer/TextureCompressor/Private/TextureCompressorModule.cpp
@@ -3238,21 +3238,37 @@ static bool CompressMipChain(
 // only useful for normal maps, fixed bad input (denormalized normals) and improved quality (quantization artifacts)
 static void NormalizeMip(FImage& InOutMip)
 {
+	TRACE_CPUPROFILER_EVENT_SCOPE(Texture.NormalizeMip);
+
+	//@todo Oodle : change FVector to FVector3f all over Texture code!
 	const FVector NormalIfZero(0.f,0.f,1.f);
 
-	const int64 NumPixels = InOutMip.GetNumPixels();
-	TArrayView64<FLinearColor> ImageColors = InOutMip.AsRGBA32F();
-	for(int64 CurPixelIndex = 0; CurPixelIndex < NumPixels; ++CurPixelIndex)
+	int64 NumPixels = InOutMip.GetNumPixels();
+	int64 NumPixelsEachJob;
+	int32 NumJobs = ImageParallelForComputeNumJobsForPixels(NumPixelsEachJob,NumPixels);
+	
+	FLinearColor * ImageColors = InOutMip.AsRGBA32F().GetData();
+
+	ParallelFor( TEXT("Texture.NormalizeMip.PF"),NumJobs,1, [&](int32 Index)
 	{
-		FLinearColor& Color = ImageColors[CurPixelIndex];
+		int64 StartIndex = Index * NumPixelsEachJob;
+		int64 EndIndex = FMath::Min(StartIndex + NumPixelsEachJob, NumPixels);
+		for (int64 i = StartIndex; i < EndIndex; ++i)
+		{
+			// @todo Oodle: SIMD NormalizeColors
 
-		FVector Normal = FVector(Color.R * 2.0f - 1.0f, Color.G * 2.0f - 1.0f, Color.B * 2.0f - 1.0f);
+			FLinearColor& Color = ImageColors[i];
 
-		// GetSafeNormal returns Vec(0,0,0) by default for tiny input, instead return flat/up
-		Normal = Normal.GetSafeNormal(UE_SMALL_NUMBER,NormalIfZero);
+			// this uses doubles due to FVector :(
 
-		Color = FLinearColor(Normal.X * 0.5f + 0.5f, Normal.Y * 0.5f + 0.5f, Normal.Z * 0.5f + 0.5f, Color.A);
-	}
+			FVector Normal(Color.R * 2.0f - 1.0f, Color.G * 2.0f - 1.0f, Color.B * 2.0f - 1.0f);
+
+			// GetSafeNormal returns Vec(0,0,0) by default for tiny input, instead return flat/up
+			Normal = Normal.GetSafeNormal(UE_SMALL_NUMBER,NormalIfZero);
+
+			Color = FLinearColor(Normal.X * 0.5f + 0.5f, Normal.Y * 0.5f + 0.5f, Normal.Z * 0.5f + 0.5f, Color.A);
+		}
+	});
 }
 
 
@@ -3971,10 +3987,32 @@ private:
 		// Apply post-mip generation adjustments.
 		if ( BuildSettings.bNormalizeNormals )
 		{
-			for ( FImage& MipImage : OutMipChain )
+			TRACE_CPUPROFILER_EVENT_SCOPE(Texture.NormalizeMipChain);
+			
+			// Parallel on the mips, but not on tiny mips :
+			int32 NumParallelMips = OutMipChain.Num();
+			while ( NumParallelMips > 0 && OutMipChain[NumParallelMips-1].GetNumPixels() < 4096 )
 			{
-				NormalizeMip(MipImage);
+				NumParallelMips--;
 			}
+
+			ParallelFor( TEXT("Texture.NormalizeMipChain.PF"),NumParallelMips,1, [&](int32 Mip)
+			{
+				if ( Mip == NumParallelMips-1 )
+				{
+					// tiny mips in the tail done on one job:
+					do
+					{
+						NormalizeMip(OutMipChain[Mip]);
+						++Mip;
+					} while( Mip < OutMipChain.Num() );
+				}
+				else
+				{
+					NormalizeMip(OutMipChain[Mip]);
+				}
+			});
+
 		}
 	
 		if (BuildSettings.bReplicateRed)
