TextureCompressorModule : trace scope slow MipHashBuilder step

diff --git a/Engine/Source/Developer/TextureCompressor/Private/TextureCompressorModule.cpp b/Engine/Source/Developer/TextureCompressor/Private/TextureCompressorModule.cpp
--- a/Engine/Source/Developer/TextureCompressor/Private/TextureCompressorModule.cpp
+++ b/Engine/Source/Developer/TextureCompressor/Private/TextureCompressorModule.cpp
@@ -3496,6 +3496,8 @@ public:
 		// apply roughness adjustment depending on normal map variation
 		if (AssociatedNormalSourceMips.Num())
 		{
+			TRACE_CPUPROFILER_EVENT_SCOPE(Texture.AssociatedNormalSourceMips);
+
 			// ECompositeTextureMode is only NormalRoughness
 			//  AssociatedNormalSourceMips should be a normal map
 
@@ -3545,6 +3547,8 @@ public:
 		
 		if (OutMetadata)
 		{
+			TRACE_CPUPROFILER_EVENT_SCOPE(Texture.MipHashBuilder);
+
 			// The metadata is about trying to determine bImageHasAlphaChannel _before_ we launch a task, which 
 			// means it must be on the actual source mips, not the post-processed mips. At some point enough textures
 			// will have this saved as part of the creation process and we can rely on it - but for now we route it
