TextureCompressor better core utilization from improved ParallelFor usage
also detach the MipChain free to a task

diff --git a/Engine/Source/Developer/TextureCompressor/Private/TextureCompressorModule.cpp b/Engine/Source/Developer/TextureCompressor/Private/TextureCompressorModule.cpp
--- a/Engine/Source/Developer/TextureCompressor/Private/TextureCompressorModule.cpp
+++ b/Engine/Source/Developer/TextureCompressor/Private/TextureCompressorModule.cpp
@@ -463,7 +463,7 @@ static FVector4f ComputeAlphaCoverage(const FVector4f Thresholds, const FVector4
 			}
 
 			FPlatformAtomics::InterlockedAdd(&CommonResult, LocalCoverage);
-		});
+		}, EParallelForFlags::Unbalanced);
 
 		Coverage[3] = float(CommonResult) / float( (int64) SourceImageData.SizeX * SourceImageData.SizeY);
 		
@@ -515,7 +515,7 @@ static FVector4f ComputeAlphaCoverage(const FVector4f Thresholds, const FVector4
 			{
 				FPlatformAtomics::InterlockedAdd(&CommonResults[i], LocalCoverage[i]);
 			}
-		});
+		}, EParallelForFlags::Unbalanced);
 
 		for (int32 i = 0; i < 4; ++i)
 		{
@@ -619,7 +619,7 @@ static void GenerateMip2x2Simple(
 				DestRow += 4;
 			}
 		}
-	});
+	}, EParallelForFlags::Unbalanced);
 }
 
 template <EMipGenAddressMode AddressMode>
@@ -649,7 +649,7 @@ static void GenerateMipUnfiltered(const FImageView2D& SourceImageData, FImageVie
 				VectorStore(FilteredColor, &DestImageData.Access(DestX, DestY).Component(0));
 			}
 		}
-	});
+	}, EParallelForFlags::Unbalanced);
 }
 
 template <EMipGenAddressMode AddressMode>
@@ -684,7 +684,7 @@ static void GenerateMip2x2(const FImageView2D& SourceImageData, FImageView2D& De
 				VectorStore(FilteredColor, &DestImageData.Access(DestX, DestY).Component(0));
 			}
 		}
-	});
+	}, EParallelForFlags::Unbalanced);
 }
 
 static void AllocateTempForMips(
@@ -831,7 +831,7 @@ static void GenerateMipSharpened(
 				}
 			}
 		}
-	});
+	}, EParallelForFlags::Unbalanced);
 }
 
 template <EMipGenAddressMode AddressMode, int KernelSize = 0>
@@ -903,7 +903,7 @@ static void GenerateMipSharpenedSeparable(
 				VectorStore(Color, &DestImageData.Access(DestX, DestY).Component(0));
 			}
 		}
-	});
+	}, EParallelForFlags::Unbalanced);
 }
 
 /**
@@ -1711,6 +1711,10 @@ void ITextureCompressorModule::GenerateMipChain(
 			FImageView2D DestView(DestImage, SliceIndex);
 			FImageView2D IntermediateDstView(IntermediateDst, SliceIndex);
 
+			// @@CB these two mipgens can run in parallel I believe?
+			//	and the one to the output mip is not blocking
+			//	we only need to block on the downsample one
+
 			// DestView is the output mip
 			GenerateSharpenedMipB8G8R8A8(
 				IntermediateSrcView, 
@@ -1758,6 +1762,8 @@ void ITextureCompressorModule::GenerateMipChain(
 
 		if ( Settings.bDownsampleWithAverage == false )
 		{
+			TRACE_CPUPROFILER_EVENT_SCOPE(Texture.GenerateMipChain.memcpy);
+
 			FMemory::Memcpy( (&IntermediateDst.AsRGBA32F()[0]), (&DestImage.AsRGBA32F()[0]),
 				IntermediateDst.SizeX * IntermediateDst.SizeY * IntermediateDst.NumSlices * sizeof(FLinearColor) );
 		}
@@ -1797,6 +1803,15 @@ void ITextureCompressorModule::GenerateMipChain(
 		IntermediateDstPtr->SizeY = FMath::Max<uint32>( 1, IntermediateSrcPtr->SizeY >> 1 );
 		IntermediateDstPtr->NumSlices = Settings.bVolume ? FMath::Max<uint32>( 1, IntermediateSrcPtr->NumSlices >> 1 ) : SrcNumSlices;
 	}
+
+	{
+		TRACE_CPUPROFILER_EVENT_SCOPE(Texture.GenerateMipChain.Free);
+		
+		// @@CB this is significant; detach trick again?
+
+		FirstTempImage = FImage();
+		SecondTempImage = FImage();
+	}
 }
 
 /*------------------------------------------------------------------------------
@@ -1993,7 +2008,7 @@ void ITextureCompressorModule::GenerateBaseCubeMipFromLongitudeLatitude2D(FImage
 					MipView.Access(x, y) = LongLatView.LookupLongLatDouble(DirectionWS);
 				}
 			}
-		});
+		}, EParallelForFlags::Unbalanced);
 	}
 }
 
@@ -2719,7 +2734,7 @@ void ITextureCompressorModule::AdjustImageColors(FImage& Image, const FTextureBu
 			{
 				AdjustColorsOld(First, Count, InBuildSettings);
 			}
-		} );
+		}, EParallelForFlags::Unbalanced);
 	}
 }
 
@@ -3271,7 +3286,7 @@ static void NormalizeMip(FImage& InOutMip)
 
 			Color = FLinearColor(Normal.X * 0.5f + 0.5f, Normal.Y * 0.5f + 0.5f, Normal.Z * 0.5f + 0.5f, Color.A);
 		}
-	});
+	}, EParallelForFlags::Unbalanced);
 }
 
 
@@ -3638,9 +3653,23 @@ public:
 		{
 			TRACE_CPUPROFILER_EVENT_SCOPE(Texture.Free_IntermediateMipChain);
 
-			// @@CB this is quite slow; detach to an async task?
+			// this is quite slow; detach to an async task?
+			//IntermediateMipChain.Empty();
 
-			IntermediateMipChain.Empty();
+			TArray<FImage>* PtrIntermediateMipChain = new TArray<FImage>( MoveTemp(IntermediateMipChain) );
+			
+			//TArray<FImage>* PtrIntermediateMipChain = new TArray<FImage>;
+			//Swap(*PtrIntermediateMipChain,IntermediateMipChain);
+
+			UE::Tasks::Launch(TEXT("Texture.Free_IntermediateMipChain.Task"),
+				[PtrIntermediateMipChain]() // by value, not ref
+				{
+					TRACE_CPUPROFILER_EVENT_SCOPE(Texture.Free_IntermediateMipChain.Task);
+					PtrIntermediateMipChain->Empty();
+					delete PtrIntermediateMipChain;
+					return true;
+				},
+				LowLevelTasks::ETaskPriority::BackgroundNormal);
 		}
 
 		return bCompressSucceeded;
@@ -4088,7 +4117,7 @@ private:
 				{
 					NormalizeMip(OutMipChain[Mip]);
 				}
-			});
+			}, EParallelForFlags::Unbalanced);
 
 		}
 	
diff --git a/Engine/Source/Runtime/Engine/Private/TextureDerivedData.cpp b/Engine/Source/Runtime/Engine/Private/TextureDerivedData.cpp
--- a/Engine/Source/Runtime/Engine/Private/TextureDerivedData.cpp
+++ b/Engine/Source/Runtime/Engine/Private/TextureDerivedData.cpp
@@ -1183,6 +1183,8 @@ void UnpackTextureBuildMetadataFromPlatformData(UE::TextureBuildUtilities::FText
  */
 int64 PutDerivedDataInCache(FTexturePlatformData* DerivedData, const FString& DerivedDataKeySuffix, const FStringView& TextureName, bool bForceAllMipsToBeInlined, bool bReplaceExistingDDC)
 {
+	TRACE_CPUPROFILER_EVENT_SCOPE(Texture.PutDerivedDataInCache);
+
 	TArray64<uint8> RawDerivedData;
 	FString DerivedDataKey;
 	int64 TotalBytesPut = 0;
diff --git a/Engine/Source/Runtime/Engine/Private/TextureDerivedDataTask.h b/Engine/Source/Runtime/Engine/Private/TextureDerivedDataTask.h
--- a/Engine/Source/Runtime/Engine/Private/TextureDerivedDataTask.h
+++ b/Engine/Source/Runtime/Engine/Private/TextureDerivedDataTask.h
@@ -103,6 +103,8 @@ struct FTextureSourceData
 
 	void ReleaseMemory()
 	{
+		TRACE_CPUPROFILER_EVENT_SCOPE(Texture.ReleaseMemory);
+
 		// Unload BulkData loaded with LoadBulkDataWithFileReader
 		AsyncSource.RemoveBulkData();
 		Blocks.Empty();
diff --git a/Engine/Source/Runtime/ImageCore/Private/ImageCore.cpp b/Engine/Source/Runtime/ImageCore/Private/ImageCore.cpp
--- a/Engine/Source/Runtime/ImageCore/Private/ImageCore.cpp
+++ b/Engine/Source/Runtime/ImageCore/Private/ImageCore.cpp
@@ -47,14 +47,42 @@ static inline int32 ParallelForComputeNumJobs(int64 & OutNumItemsPerJob,int64 Nu
 		return 1;
 	}
 	
-	// ParallelFor will actually make 6*NumWorkers batches and then make NumWorkers tasks that pop the batches
-	//	this helps with mismatched thread runtime
-	//	here we only make NumWorkers batches max
-	//	but this is rarely a problem in image cook because it is parallelized already at a the higher level
+	// Always use ParallelFor with "Unbalanced"  so it does not do weird things to your job count
 
 	const int32 NumWorkers = FMath::Max(1, FTaskGraphInterface::Get().GetNumWorkerThreads());
 	int32 NumJobs = (int32)(NumItems / MinNumItemsPerJob); // round down
-	NumJobs = FMath::Clamp(NumJobs, int32(1), NumWorkers); 
+
+	// NumJobs is now the maximum number of jobs we would want to do
+	//	(eg. one for each small packet of MinNumItemsPerJob)
+	//	but that's often way too many
+
+	// Assign 1,2, or 3 jobs per worker
+	// never make jobs that will be smaller than MinNumItemsPerJob (eg. NumJobs is only adjusted downward)
+	// more jobs per worker provides load balancing, if one worker stalls others can pick up the jobs he didn't do
+
+	if ( NumJobs >= NumWorkers*3 )
+	{
+		// enough jobs to evenly distribute 3 per worker
+		// reduce job count to provide some slack for stalling workers
+		NumJobs = NumWorkers*3 - 1;
+	}
+	else if ( NumJobs >= NumWorkers*2 )
+	{
+		// enough jobs to evenly distribute 2 per worker
+		NumJobs = NumWorkers*2 - 1;
+	}
+	else
+	{
+		// not enough jobs to evenly distribute more than 1 per worker
+		NumJobs = FMath::Clamp(NumJobs, int32(1), NumWorkers); 
+	}
+
+	if ( NumJobs > 16 )
+	{
+		// add some slack; assume one thread will not be available
+		//	slightly under-subscribing reduces the chance of a much higher latency
+		NumJobs--;
+	}
 
 	OutNumItemsPerJob = (NumItems + NumJobs-1) / NumJobs; // round up
 	check( NumJobs*OutNumItemsPerJob >= NumItems );
@@ -103,7 +131,7 @@ static void ParallelLoop(const TCHAR* DebugName, int32 NumJobs, int64 TexelsPerJ
 		{
 			Func(TexelIndex);
 		}
-	});
+	}, EParallelForFlags::Unbalanced);
 }
 
 // ParallelOr : call Func() on all texels ; returns true if Func is true for any texel
@@ -128,7 +156,7 @@ static bool ParallelOr(const TCHAR* DebugName, int32 NumJobs, int64 TexelsPerJob
 		// if Func was really slow, you might want to check if merged_bool was set to true on some other thread
 		//	and terminate our loop when that is true
 		//  but that causes cache traffic for the shared variable, so should not be done in tight loops
-	});
+	}, EParallelForFlags::Unbalanced);
 
 	bool ret = merged_bool.load(std::memory_order_acquire);
 	return ret;
@@ -158,7 +186,7 @@ static void ParallelLoop(const TCHAR* DebugName, int32 NumJobs, int64 TexelsPerJ
 			Func(TexelIndex);
 			++TexelIndex;
 		}
-	});
+	}, EParallelForFlags::Unbalanced);
 }
 
 
@@ -363,7 +391,7 @@ IMAGECORE_API void FImageCore::CopyImage(const FImageView & SrcImage,const FImag
 							const int64 Count = FMath::Min(TexelsPerJob, NumTexels - StartIndex);
 							ConvertFLinearColorsToFColorSRGB(&SrcColors[StartIndex], &DestColors[StartIndex], Count);
 						}
-					);
+					, EParallelForFlags::Unbalanced);
 				}
 				else
 				{
