TextureCompressorModule : no need to run the same 2x2 mip gen filter twice

#rnx

diff --git a/Engine/Source/Developer/TextureCompressor/Private/TextureCompressorModule.cpp b/Engine/Source/Developer/TextureCompressor/Private/TextureCompressorModule.cpp
--- a/Engine/Source/Developer/TextureCompressor/Private/TextureCompressorModule.cpp
+++ b/Engine/Source/Developer/TextureCompressor/Private/TextureCompressorModule.cpp
@@ -1601,7 +1601,7 @@ void ITextureCompressorModule::GenerateMipChain(
 		return;
 	}
 
-	check(BaseImage.Format == ERawImageFormat::RGBA32F);
+	check(BaseImage.Format == ERawImageFormat::RGBA32F && BaseImage.GammaSpace == EGammaSpace::Linear);
 
 	const FImage& BaseMip = BaseImage;
 	const int32 SrcWidth = BaseMip.SizeX;
@@ -1615,21 +1615,11 @@ void ITextureCompressorModule::GenerateMipChain(
 	// This will be used as a buffer for the mip processing
 	FImage FirstTempImage;
 
-	if (BaseMip.GammaSpace != EGammaSpace::Linear)
-	{
-		// copy base mip
-		BaseMip.CopyTo(FirstTempImage, ERawImageFormat::RGBA32F, EGammaSpace::Linear);
-
-		IntermediateSrcPtr = &FirstTempImage;
-	}
-	else
-	{
-		// It looks like the BaseMip can be reused for the intermediate source of the second Mip (assuming that the format was check earlier to be RGBA32F)
-		IntermediateSrcPtr = &BaseMip;
+	// Source for first mip step is the passed-in image:
+	IntermediateSrcPtr = &BaseMip;
 
-		// This temp image will be first used as an intermediate destination for the third mip in the chain
-		FirstTempImage.Init( FMath::Max<uint32>( 1, SrcWidth >> 2 ), FMath::Max<uint32>( 1, SrcHeight >> 2 ), Settings.bVolume ? FMath::Max<uint32>( 1, SrcNumSlices >> 2 ) : SrcNumSlices, ImageFormat );
-	}
+	// This temp image will be first used as an intermediate destination for the third mip in the chain
+	FirstTempImage.Init( FMath::Max<uint32>( 1, SrcWidth >> 2 ), FMath::Max<uint32>( 1, SrcHeight >> 2 ), Settings.bVolume ? FMath::Max<uint32>( 1, SrcNumSlices >> 2 ) : SrcNumSlices, ImageFormat );
 
 	// The image for the first destination
 	FImage SecondTempImage(FMath::Max<uint32>( 1, SrcWidth >> 1 ), FMath::Max<uint32>( 1, SrcHeight >> 1 ), Settings.bVolume ? FMath::Max<uint32>( 1, SrcNumSlices >> 1 ) : SrcNumSlices, ImageFormat);
@@ -1641,6 +1631,15 @@ void ITextureCompressorModule::GenerateMipChain(
 	KernelSimpleAverage.BuildSeparatableGaussWithSharpen( 2 );
 	KernelDownsample.BuildSeparatableGaussWithSharpen( Settings.SharpenMipKernelSize, Settings.MipSharpening );
 
+	bool bDownsampleWithAverage = Settings.bDownsampleWithAverage;
+	if ( Settings.SharpenMipKernelSize == 2 )
+	{
+		// filter is simple 2x2
+		//	no need to do the primary filter and also a 2x2 downfilter
+		//	if the primary filter IS 2x2
+		bDownsampleWithAverage = false;
+	}
+
 	//@TODO : add a true 3D kernel.
 
 	EMipGenAddressMode AddressMode = ComputeAddressMode(*IntermediateSrcPtr,Settings);
@@ -1669,7 +1668,11 @@ void ITextureCompressorModule::GenerateMipChain(
 	const bool bUnfiltered = Settings.MipGenSettings == TMGS_Unfiltered;
 	const bool bUseNewMipFilter = Settings.bUseNewMipFilter;
 	AllocateTempForMips(DownsampleTempData, BaseMip.SizeX, BaseMip.SizeY, SecondTempImage.SizeX, SecondTempImage.SizeY, bDoScaleMipsForAlphaCoverage, KernelDownsample, 2, bSharpenWithoutColorShift, bUnfiltered, bUseNewMipFilter);
-	AllocateTempForMips(AverageTempData, BaseMip.SizeX, BaseMip.SizeY, SecondTempImage.SizeX, SecondTempImage.SizeY, bDoScaleMipsForAlphaCoverage, KernelSimpleAverage, 2, bSharpenWithoutColorShift, bUnfiltered, bUseNewMipFilter);
+
+	if ( bDownsampleWithAverage )
+	{
+		AllocateTempForMips(AverageTempData, BaseMip.SizeX, BaseMip.SizeY, SecondTempImage.SizeX, SecondTempImage.SizeY, bDoScaleMipsForAlphaCoverage, KernelSimpleAverage, 2, bSharpenWithoutColorShift, bUnfiltered, bUseNewMipFilter);
+	}
 
 	// Generate mips
 	//  default value of MipChainDepth is MAX_uint32, means generate all mips down to 1x1
@@ -1677,9 +1680,18 @@ void ITextureCompressorModule::GenerateMipChain(
 	for (; MipChainDepth != 0 ; --MipChainDepth)
 	{
 		check(IntermediateSrcPtr && IntermediateDstPtr);
+
+		// IntermediateSrc & Dest = First/Second Temp Image, swapping at each step
+		// IntermediateSrc/Dest are used to generate down the chain
+		
 		const FImage& IntermediateSrc = *IntermediateSrcPtr;
 		FImage& IntermediateDst = *IntermediateDstPtr;
 		
+		// Update the destination size for the next iteration.
+		IntermediateDst.SizeX = FMath::Max<uint32>( 1, IntermediateSrc.SizeX >> 1 );
+		IntermediateDst.SizeY = FMath::Max<uint32>( 1, IntermediateSrc.SizeY >> 1 );
+		IntermediateDst.NumSlices = Settings.bVolume ? FMath::Max<uint32>( 1, IntermediateSrc.NumSlices >> 1 ) : SrcNumSlices;
+
 		if ( IntermediateSrc.SizeX == 1 && IntermediateSrc.SizeY == 1 && (!Settings.bVolume || IntermediateSrc.NumSlices == 1))
 		{
 			// should not have been called, starting mip is already small enough
@@ -1711,8 +1723,9 @@ void ITextureCompressorModule::GenerateMipChain(
 			FImageView2D DestView(DestImage, SliceIndex);
 			FImageView2D IntermediateDstView(IntermediateDst, SliceIndex);
 
-			// @@CB these two mipgens can run in parallel I believe?
-			//	and the one to the output mip is not blocking
+			// @@CB 
+			// if ( bDownsampleWithAverage ) these two mipgens can run in parallel I believe?
+			//	and the one to the output mip is not blocking ; <- no, not quite true because we reuse the *source* for the pingpong
 			//	we only need to block on the downsample one
 
 			// DestView is the output mip
@@ -1733,9 +1746,9 @@ void ITextureCompressorModule::GenerateMipChain(
 
 			// generate IntermediateDstImage:
 			// IntermediateDstImage will be the source for the next mip
-			if ( Settings.bDownsampleWithAverage )
+			if ( bDownsampleWithAverage )
 			{
-				// down sample without sharpening for the next iteration
+				// down sample with 2x2 for the next iteration
 				// bDownsampleWithAverage comes from GetMipGenSettings
 				// it is on by default for all cases except Blur
 				// it means every mip is generated *twice*
@@ -1760,10 +1773,13 @@ void ITextureCompressorModule::GenerateMipChain(
 			}
 		}
 
-		if ( Settings.bDownsampleWithAverage == false )
+		if ( ! bDownsampleWithAverage )
 		{
 			TRACE_CPUPROFILER_EVENT_SCOPE(Texture.GenerateMipChain.memcpy);
 
+			// @@ CB this is unnecessary, should just make the next mip from DestImage instead of IntermediateDest
+			//	in fact in this case you don't need the Intermediate ping-pong at all, just walk down the mip chain directly
+
 			FMemory::Memcpy( (&IntermediateDst.AsRGBA32F()[0]), (&DestImage.AsRGBA32F()[0]),
 				IntermediateDst.SizeX * IntermediateDst.SizeY * IntermediateDst.NumSlices * sizeof(FLinearColor) );
 		}
@@ -1797,11 +1813,6 @@ void ITextureCompressorModule::GenerateMipChain(
 			IntermediateDstPtr = &SecondTempImage;
 			IntermediateSrcPtr = &FirstTempImage;
 		}
-
-		// Update the destination size for the next iteration.
-		IntermediateDstPtr->SizeX = FMath::Max<uint32>( 1, IntermediateSrcPtr->SizeX >> 1 );
-		IntermediateDstPtr->SizeY = FMath::Max<uint32>( 1, IntermediateSrcPtr->SizeY >> 1 );
-		IntermediateDstPtr->NumSlices = Settings.bVolume ? FMath::Max<uint32>( 1, IntermediateSrcPtr->NumSlices >> 1 ) : SrcNumSlices;
 	}
 
 	{
