TextureCompressorModule : GenerateMipChain no need to use ping pong buffers and memcpy if the output and average filter are the same, we can work directly in the output mip chain

#rnx

diff --git a/Engine/Source/Developer/TextureCompressor/Private/TextureCompressorModule.cpp b/Engine/Source/Developer/TextureCompressor/Private/TextureCompressorModule.cpp
--- a/Engine/Source/Developer/TextureCompressor/Private/TextureCompressorModule.cpp
+++ b/Engine/Source/Developer/TextureCompressor/Private/TextureCompressorModule.cpp
@@ -1013,6 +1013,11 @@
 	}
 }
 
+static void GenerateMipBorder(
+	const FImageView2D& SrcImageData, 
+	FImageView2D& DestImageData
+	);
+
 // to switch conveniently between different texture wrapping modes for the mip map generation
 // the template can optimize the inner loop using a constant AddressMode
 static void GenerateSharpenedMipB8G8R8A8(
@@ -1028,7 +1033,8 @@
 	uint32 ScaleFactor,
 	bool bSharpenWithoutColorShift,
 	bool bUnfiltered,
-	bool bUseNewMipFilter
+	bool bUseNewMipFilter,
+	bool bReDrawBorder = false
 	)
 {
 	TRACE_CPUPROFILER_EVENT_SCOPE(Texture.GenerateSharpenedMip);
@@ -1076,6 +1082,11 @@
 				(DestImageData.SliceColors[ColorIndex] + TempImageData.SliceColors[ColorIndex]) * 0.5f;
 		}
 	}
+
+	if ( bReDrawBorder )
+	{
+		GenerateMipBorder( SourceImageData, DestImageData );
+	}
 }
 
 // Update border texels after normal mip map generation to preserve the colors there (useful for particles and decals).
@@ -1585,11 +1596,12 @@
 	}
 }
 
-void ITextureCompressorModule::GenerateMipChain(
+
+static void GenerateMipChain_Old_Reference(
 	const FTextureBuildSettings& Settings,
 	const FImage& BaseImage,
 	TArray<FImage> &OutMipChain,
-	uint32 MipChainDepth 
+	uint32 MipChainDepth
 	)
 {
 	TRACE_CPUPROFILER_EVENT_SCOPE(Texture.GenerateMipChain);
@@ -1601,7 +1613,7 @@
 		return;
 	}
 
-	check(BaseImage.Format == ERawImageFormat::RGBA32F && BaseImage.GammaSpace == EGammaSpace::Linear);
+	check(BaseImage.Format == ERawImageFormat::RGBA32F);
 
 	const FImage& BaseMip = BaseImage;
 	const int32 SrcWidth = BaseMip.SizeX;
@@ -1615,11 +1627,21 @@
 	// This will be used as a buffer for the mip processing
 	FImage FirstTempImage;
 
-	// Source for first mip step is the passed-in image:
-	IntermediateSrcPtr = &BaseMip;
+	if (BaseMip.GammaSpace != EGammaSpace::Linear)
+	{
+		// copy base mip
+		BaseMip.CopyTo(FirstTempImage, ERawImageFormat::RGBA32F, EGammaSpace::Linear);
 
-	// This temp image will be first used as an intermediate destination for the third mip in the chain
-	FirstTempImage.Init( FMath::Max<uint32>( 1, SrcWidth >> 2 ), FMath::Max<uint32>( 1, SrcHeight >> 2 ), Settings.bVolume ? FMath::Max<uint32>( 1, SrcNumSlices >> 2 ) : SrcNumSlices, ImageFormat );
+		IntermediateSrcPtr = &FirstTempImage;
+	}
+	else
+	{
+		// It looks like the BaseMip can be reused for the intermediate source of the second Mip (assuming that the format was check earlier to be RGBA32F)
+		IntermediateSrcPtr = &BaseMip;
 
+		// This temp image will be first used as an intermediate destination for the third mip in the chain
+		FirstTempImage.Init( FMath::Max<uint32>( 1, SrcWidth >> 2 ), FMath::Max<uint32>( 1, SrcHeight >> 2 ), Settings.bVolume ? FMath::Max<uint32>( 1, SrcNumSlices >> 2 ) : SrcNumSlices, ImageFormat );
+	}
+
 	// The image for the first destination
 	FImage SecondTempImage(FMath::Max<uint32>( 1, SrcWidth >> 1 ), FMath::Max<uint32>( 1, SrcHeight >> 1 ), Settings.bVolume ? FMath::Max<uint32>( 1, SrcNumSlices >> 1 ) : SrcNumSlices, ImageFormat);
@@ -1631,15 +1653,6 @@
 	KernelSimpleAverage.BuildSeparatableGaussWithSharpen( 2 );
 	KernelDownsample.BuildSeparatableGaussWithSharpen( Settings.SharpenMipKernelSize, Settings.MipSharpening );
 
-	bool bDownsampleWithAverage = Settings.bDownsampleWithAverage;
-	if ( Settings.SharpenMipKernelSize == 2 )
-	{
-		// filter is simple 2x2
-		//	no need to do the primary filter and also a 2x2 downfilter
-		//	if the primary filter IS 2x2
-		bDownsampleWithAverage = false;
-	}
-
 	//@TODO : add a true 3D kernel.
 
 	EMipGenAddressMode AddressMode = ComputeAddressMode(*IntermediateSrcPtr,Settings);
@@ -1668,11 +1681,7 @@
 	const bool bUnfiltered = Settings.MipGenSettings == TMGS_Unfiltered;
 	const bool bUseNewMipFilter = Settings.bUseNewMipFilter;
 	AllocateTempForMips(DownsampleTempData, BaseMip.SizeX, BaseMip.SizeY, SecondTempImage.SizeX, SecondTempImage.SizeY, bDoScaleMipsForAlphaCoverage, KernelDownsample, 2, bSharpenWithoutColorShift, bUnfiltered, bUseNewMipFilter);
+	AllocateTempForMips(AverageTempData, BaseMip.SizeX, BaseMip.SizeY, SecondTempImage.SizeX, SecondTempImage.SizeY, bDoScaleMipsForAlphaCoverage, KernelSimpleAverage, 2, bSharpenWithoutColorShift, bUnfiltered, bUseNewMipFilter);
 
-	if ( bDownsampleWithAverage )
-	{
-		AllocateTempForMips(AverageTempData, BaseMip.SizeX, BaseMip.SizeY, SecondTempImage.SizeX, SecondTempImage.SizeY, bDoScaleMipsForAlphaCoverage, KernelSimpleAverage, 2, bSharpenWithoutColorShift, bUnfiltered, bUseNewMipFilter);
-	}
-
 	// Generate mips
 	//  default value of MipChainDepth is MAX_uint32, means generate all mips down to 1x1
@@ -1680,18 +1689,9 @@
 	for (; MipChainDepth != 0 ; --MipChainDepth)
 	{
 		check(IntermediateSrcPtr && IntermediateDstPtr);
-
-		// IntermediateSrc & Dest = First/Second Temp Image, swapping at each step
-		// IntermediateSrc/Dest are used to generate down the chain
-		
 		const FImage& IntermediateSrc = *IntermediateSrcPtr;
 		FImage& IntermediateDst = *IntermediateDstPtr;
 		
-		// Update the destination size for the next iteration.
-		IntermediateDst.SizeX = FMath::Max<uint32>( 1, IntermediateSrc.SizeX >> 1 );
-		IntermediateDst.SizeY = FMath::Max<uint32>( 1, IntermediateSrc.SizeY >> 1 );
-		IntermediateDst.NumSlices = Settings.bVolume ? FMath::Max<uint32>( 1, IntermediateSrc.NumSlices >> 1 ) : SrcNumSlices;
-
 		if ( IntermediateSrc.SizeX == 1 && IntermediateSrc.SizeY == 1 && (!Settings.bVolume || IntermediateSrc.NumSlices == 1))
 		{
 			// should not have been called, starting mip is already small enough
@@ -1723,10 +1723,7 @@
 			FImageView2D DestView(DestImage, SliceIndex);
 			FImageView2D IntermediateDstView(IntermediateDst, SliceIndex);
 
-			// @@CB 
-			// if ( bDownsampleWithAverage ) these two mipgens can run in parallel I believe?
-			//	and the one to the output mip is not blocking ; <- no, not quite true because we reuse the *source* for the pingpong
-			//	we only need to block on the downsample one
+			// GenerateMipChain_Old_Reference
 
 			// DestView is the output mip
 			GenerateSharpenedMipB8G8R8A8(
@@ -1746,9 +1743,9 @@
 
 			// generate IntermediateDstImage:
 			// IntermediateDstImage will be the source for the next mip
-			if ( bDownsampleWithAverage )
+			if ( Settings.bDownsampleWithAverage )
 			{
-				// down sample with 2x2 for the next iteration
+				// down sample without sharpening for the next iteration
 				// bDownsampleWithAverage comes from GetMipGenSettings
 				// it is on by default for all cases except Blur
 				// it means every mip is generated *twice*
@@ -1773,13 +1770,10 @@
 			}
 		}
 
-		if ( ! bDownsampleWithAverage )
+		if ( Settings.bDownsampleWithAverage == false )
 		{
 			TRACE_CPUPROFILER_EVENT_SCOPE(Texture.GenerateMipChain.memcpy);
 
-			// @@ CB this is unnecessary, should just make the next mip from DestImage instead of IntermediateDest
-			//	in fact in this case you don't need the Intermediate ping-pong at all, just walk down the mip chain directly
-
 			FMemory::Memcpy( (&IntermediateDst.AsRGBA32F()[0]), (&DestImage.AsRGBA32F()[0]),
 				IntermediateDst.SizeX * IntermediateDst.SizeY * IntermediateDst.NumSlices * sizeof(FLinearColor) );
 		}
@@ -1813,6 +1807,275 @@
 			IntermediateDstPtr = &SecondTempImage;
 			IntermediateSrcPtr = &FirstTempImage;
 		}
+
+		// Update the destination size for the next iteration.
+		IntermediateDstPtr->SizeX = FMath::Max<uint32>( 1, IntermediateSrcPtr->SizeX >> 1 );
+		IntermediateDstPtr->SizeY = FMath::Max<uint32>( 1, IntermediateSrcPtr->SizeY >> 1 );
+		IntermediateDstPtr->NumSlices = Settings.bVolume ? FMath::Max<uint32>( 1, IntermediateSrcPtr->NumSlices >> 1 ) : SrcNumSlices;
 	}
 
 	{
+		TRACE_CPUPROFILER_EVENT_SCOPE(Texture.GenerateMipChain.Free);
+		
+		FirstTempImage = FImage();
+		SecondTempImage = FImage();
+	}
+} // GenerateMipChain_Old_Reference
+
+void ITextureCompressorModule::GenerateMipChain(
+	const FTextureBuildSettings& Settings,
+	const FImage& BaseImage,
+	TArray<FImage> &OutMipChain,
+	uint32 MipChainDepth 
+	)
+{
+	TRACE_CPUPROFILER_EVENT_SCOPE(Texture.GenerateMipChain);
+	
+	// maybe someday : add a true 3D kernel for Volume Textures?
+
+	// MipChainDepth is the number more to make, OutMipChain has some already
+	// typically BaseImage == OutMipChain.Last()
+	if ( MipChainDepth == 0 )
+	{
+		return;
+	}
+
+	check(BaseImage.Format == ERawImageFormat::RGBA32F && BaseImage.GammaSpace == EGammaSpace::Linear);
+
+	const FImage& BaseMip = BaseImage;
+	const int32 SrcWidth = BaseMip.SizeX;
+	const int32 SrcHeight= BaseMip.SizeY;
+	const int32 SrcNumSlices = BaseMip.NumSlices;
+	const ERawImageFormat::Type ImageFormat = ERawImageFormat::RGBA32F;
+
+	// Filtering kernels.
+	FImageKernel2D KernelSimpleAverage;
+	FImageKernel2D KernelDownsample;
+	KernelSimpleAverage.BuildSeparatableGaussWithSharpen( 2 );
+	KernelDownsample.BuildSeparatableGaussWithSharpen( Settings.SharpenMipKernelSize, Settings.MipSharpening );
+
+	bool bDownsampleWithAverage = Settings.bDownsampleWithAverage;
+	if ( Settings.SharpenMipKernelSize == 2 )
+	{
+		// filter is simple 2x2
+		//	no need to do the primary filter and also a 2x2 downfilter
+		//	if the primary filter IS 2x2
+		bDownsampleWithAverage = false;
+	}
+	
+	/**
+	 
+	 GenerateMipChain
+
+	 has two distinct modes of operation
+
+	 if !bDownsampleWithAverage
+	 then the mip filter is used to both make the output mips, and to make the parents down the chain
+	 in this case no temp images are needed
+	 there is only one mipgen per mip level
+	 each mipgen is dependent on the previous
+	 this is used for 2x2 downsample and blurs
+
+	 if bDownsampleWithAverage
+	 then the chosen mip filter is used to make output mips from its parent
+	 but a different filter (2x2) is used to make the parents down the chain
+	 two mipgens are done per level
+	 two temp images are used as a pingpong swap of parent/child down the chain
+	 this is used for the "sharpen" family of filters
+
+	 every mip level is made twice :
+
+		[0] -> [1] -> [2] -> [3] -> .. with SimpleAverage 2x2 ; these are stored in FirstTemp/SecondTemp
+		  \      \      \
+		   \      \      \
+			[1]    [2]    [3]   .. with Sharpen ; these go in OutMipChain
+
+
+	 */
+
+	const FImage* IntermediateSrcPtr = nullptr;
+	FImage* IntermediateDstPtr = nullptr;
+
+	// This will be used as a buffer for the mip processing
+	FImage FirstTempImage;
+
+	// Source for first mip step is the passed-in image:
+	IntermediateSrcPtr = &BaseMip;
+
+	// The image for the first destination
+	FImage SecondTempImage;
+	
+	if ( bDownsampleWithAverage )
+	{
+		SecondTempImage.Init(FMath::Max<uint32>( 1, SrcWidth >> 1 ), FMath::Max<uint32>( 1, SrcHeight >> 1 ), Settings.bVolume ? FMath::Max<uint32>( 1, SrcNumSlices >> 1 ) : SrcNumSlices, ImageFormat);
+
+		// This temp image will be first used as an intermediate destination for the third mip in the chain
+		FirstTempImage.Init( FMath::Max<uint32>( 1, SrcWidth >> 2 ), FMath::Max<uint32>( 1, SrcHeight >> 2 ), Settings.bVolume ? FMath::Max<uint32>( 1, SrcNumSlices >> 2 ) : SrcNumSlices, ImageFormat );
+
+		IntermediateDstPtr = &SecondTempImage;
+	}
+
+	EMipGenAddressMode AddressMode = ComputeAddressMode(*IntermediateSrcPtr,Settings);
+	bool bReDrawBorder = false;
+	if( Settings.bPreserveBorder )
+	{
+		bReDrawBorder = !Settings.bBorderColorBlack;
+	}
+
+	// Calculate alpha coverage value to preserve along mip chain
+	FVector4f AlphaCoverages(0, 0, 0, 0);
+	if ( Settings.bDoScaleMipsForAlphaCoverage )
+	{
+		check(Settings.AlphaCoverageThresholds != FVector4f(0,0,0,0));
+		check(IntermediateSrcPtr);
+		const FImageView2D IntermediateSrcView = FImageView2D::ConstructConst(*IntermediateSrcPtr, 0);
+
+		const FVector4f AlphaScales(1, 1, 1, 1);		
+		AlphaCoverages = ComputeAlphaCoverage(Settings.AlphaCoverageThresholds, AlphaScales, IntermediateSrcView);
+	}
+
+	TArray<FLinearColor> DownsampleTempData;
+	TArray<FLinearColor> AverageTempData;
+	const bool bDoScaleMipsForAlphaCoverage = Settings.bDoScaleMipsForAlphaCoverage;
+	const bool bSharpenWithoutColorShift = Settings.bSharpenWithoutColorShift;
+	const bool bUnfiltered = Settings.MipGenSettings == TMGS_Unfiltered;
+	const bool bUseNewMipFilter = Settings.bUseNewMipFilter;
+	AllocateTempForMips(DownsampleTempData, BaseMip.SizeX, BaseMip.SizeY, SecondTempImage.SizeX, SecondTempImage.SizeY, bDoScaleMipsForAlphaCoverage, KernelDownsample, 2, bSharpenWithoutColorShift, bUnfiltered, bUseNewMipFilter);
+
+	if ( bDownsampleWithAverage )
+	{
+		AllocateTempForMips(AverageTempData, BaseMip.SizeX, BaseMip.SizeY, SecondTempImage.SizeX, SecondTempImage.SizeY, bDoScaleMipsForAlphaCoverage, KernelSimpleAverage, 2, bSharpenWithoutColorShift, bUnfiltered, bUseNewMipFilter);
+	}
+
+	// Generate mips
+	//  default value of MipChainDepth is MAX_uint32, means generate all mips down to 1x1
+	//	(break inside the loop)
+	for (; MipChainDepth != 0 ; --MipChainDepth)
+	{
+		check(IntermediateSrcPtr);
+		const FImage& IntermediateSrc = *IntermediateSrcPtr;
+		
+		if ( IntermediateSrc.SizeX == 1 && IntermediateSrc.SizeY == 1 && (!Settings.bVolume || IntermediateSrc.NumSlices == 1))
+		{
+			// should not have been called, starting mip is already small enough
+			check(0);
+			break;
+		}
+		
+		int32 DstSizeX = FMath::Max( 1, IntermediateSrc.SizeX >> 1 );
+		int32 DstSizeY = FMath::Max( 1, IntermediateSrc.SizeY >> 1 );
+		int32 DstNumSlices = Settings.bVolume ? FMath::Max( 1, IntermediateSrc.NumSlices >> 1 ) : SrcNumSlices;
+
+		if ( bDownsampleWithAverage )
+		{
+			// IntermediateSrc & Dest = First/Second Temp Image, swapping at each step
+			// IntermediateSrc/Dest are used to generate down the chain
+
+			check(IntermediateDstPtr);
+
+			// Update the destination size for the next iteration.
+			IntermediateDstPtr->SizeX = DstSizeX;
+			IntermediateDstPtr->SizeY = DstSizeY;
+			IntermediateDstPtr->NumSlices = DstNumSlices;
+		}
+
+		// add new mip to TArray<FImage> &OutMipChain :
+		//	placement new on TArray does AddUninitialized then constructs in the last element
+		check( OutMipChain.GetSlack() > 0 ); // should have been reserved
+		FImage& DestImage = *new(OutMipChain) FImage(DstSizeX, DstSizeY, DstNumSlices, ImageFormat);
+		
+		for (int32 SliceIndex = 0; SliceIndex < DstNumSlices; ++SliceIndex)
+		{
+			const int32 SrcSliceIndex = Settings.bVolume ? (SliceIndex * 2) : SliceIndex;
+			const FImageView2D IntermediateSrcView = FImageView2D::ConstructConst(IntermediateSrc, SrcSliceIndex);
+			FImageView2D IntermediateSrcView2;
+			if ( Settings.bVolume )
+			{
+				const int32 SrcSliceIndex2 = FMath::Min(SrcSliceIndex + 1,IntermediateSrc.NumSlices-1);
+				IntermediateSrcView2 = FImageView2D::ConstructConst(IntermediateSrc, SrcSliceIndex2);
+			}
+			FImageView2D DestView(DestImage, SliceIndex);
+
+			// @@CB 
+			// if ( bDownsampleWithAverage ) these two mipgens can run in parallel I believe?
+			//	and the one to the output mip is not blocking ; <- no, not quite true because we reuse the *source* for the pingpong
+			//	we only need to block on the downsample one
+
+			// DestView is the output mip
+			GenerateSharpenedMipB8G8R8A8(
+				IntermediateSrcView, 
+				IntermediateSrcView2,
+				DownsampleTempData,
+				DestView,
+				AddressMode,
+				bDoScaleMipsForAlphaCoverage,
+				AlphaCoverages,
+				Settings.AlphaCoverageThresholds,
+				KernelDownsample,
+				2,
+				bSharpenWithoutColorShift,
+				bUnfiltered,
+				bUseNewMipFilter,
+				bReDrawBorder);
+
+			// generate IntermediateDstImage:
+			// IntermediateDstImage will be the source for the next mip
+			if ( bDownsampleWithAverage )
+			{
+				// down sample with 2x2 for the next iteration
+				// bDownsampleWithAverage comes from GetMipGenSettings
+				// it is on by default for all cases except Blur
+				// it means every mip is generated *twice*
+				// the output mip is made above using Sharpen from the IntermediateSrc
+				// then the next source is made here using 2x2 ("SimpleAverage")
+				// the next IntermediateSrc is not my outputmip, it's what is made here
+				
+				FImageView2D IntermediateDstView(*IntermediateDstPtr, SliceIndex);
+
+				GenerateSharpenedMipB8G8R8A8(
+					IntermediateSrcView,
+					IntermediateSrcView2,
+					AverageTempData,
+					IntermediateDstView,
+					AddressMode,
+					bDoScaleMipsForAlphaCoverage,
+					AlphaCoverages,
+					Settings.AlphaCoverageThresholds,
+					KernelSimpleAverage,
+					2,
+					bSharpenWithoutColorShift,
+					bUnfiltered,
+					bUseNewMipFilter,
+					bReDrawBorder);
+			}
+		}
+
+		// Once we've created mip-maps down to 1x1, we're done.
+		if ( DstSizeX == 1 && DstSizeY == 1 && (!Settings.bVolume || DstNumSlices == 1))
+		{
+			break;
+		}
+		
+		if ( bDownsampleWithAverage )
+		{
+			// last destination becomes next source
+			if (IntermediateDstPtr == &SecondTempImage)
+			{
+				IntermediateDstPtr = &FirstTempImage;
+				IntermediateSrcPtr = &SecondTempImage;
+			}
+			else
+			{
+				IntermediateDstPtr = &SecondTempImage;
+				IntermediateSrcPtr = &FirstTempImage;
+			}
+		}
+		else
+		{
+			// point at the mip we made in OutMipChain
+			//  safe because OutMipChain is ensured to not reallocate
+			IntermediateSrcPtr = &DestImage;
+		}
+	}
+
+	{
@@ -4081,7 +4344,24 @@
 				// you could pass GenerateCount as the large arg here
 				//  and it should make the same result
 
-				GenerateMipChain(BuildSettings, OutMipChain.Last(), OutMipChain, MAX_uint32);
+				int StartImageIndex = OutMipChain.Num()-1;
+				GenerateMipChain(BuildSettings, OutMipChain[StartImageIndex], OutMipChain, MAX_uint32);
+
+				#if 0 // @@CB temp debug code, remove
+				{
+					// ensure that GenerateMipChain makes the exact same output as GenerateMipChain_Old_Reference :
+					TArray<FImage> ReferenceMipChain;
+					GenerateMipChain_Old_Reference(BuildSettings, OutMipChain[StartImageIndex], ReferenceMipChain, MAX_uint32);
+
+					for(int refi=0,outi=StartImageIndex+1;outi<OutMipChain.Num();outi++,refi++)
+					{
+						FImageInfo OutInfo = OutMipChain[outi];
+						FImageInfo RefInfo = ReferenceMipChain[refi];
+						check( OutInfo == RefInfo );
+						check( OutMipChain[outi].RawData == ReferenceMipChain[refi].RawData ); // runs CompareItems which becomes memcmp
+					}
+				}
+				#endif
 			}
 		}
 		check(OutMipChain.Num() == NumOutputMips);
