TextureCompressorModule : fix bug introduced on Blur+NewFilters filter MipGen

#rnx

diff --git a/Engine/Source/Developer/TextureCompressor/Private/TextureCompressorModule.cpp b/Engine/Source/Developer/TextureCompressor/Private/TextureCompressorModule.cpp
--- a/Engine/Source/Developer/TextureCompressor/Private/TextureCompressorModule.cpp
+++ b/Engine/Source/Developer/TextureCompressor/Private/TextureCompressorModule.cpp
@@ -1933,6 +1933,9 @@ void ITextureCompressorModule::GenerateMipChain(
 		const FVector4f AlphaScales(1, 1, 1, 1);		
 		AlphaCoverages = ComputeAlphaCoverage(Settings.AlphaCoverageThresholds, AlphaScales, IntermediateSrcView);
 	}
+	
+	int32 FirstMipSizeX = FMath::Max(BaseMip.SizeX>>1,1);
+	int32 FirstMipSizeY = FMath::Max(BaseMip.SizeY>>1,1);
 
 	TArray<FLinearColor> DownsampleTempData;
 	TArray<FLinearColor> AverageTempData;
@@ -1940,11 +1943,11 @@ void ITextureCompressorModule::GenerateMipChain(
 	const bool bSharpenWithoutColorShift = Settings.bSharpenWithoutColorShift;
 	const bool bUnfiltered = Settings.MipGenSettings == TMGS_Unfiltered;
 	const bool bUseNewMipFilter = Settings.bUseNewMipFilter;
-	AllocateTempForMips(DownsampleTempData, BaseMip.SizeX, BaseMip.SizeY, SecondTempImage.SizeX, SecondTempImage.SizeY, bDoScaleMipsForAlphaCoverage, KernelDownsample, 2, bSharpenWithoutColorShift, bUnfiltered, bUseNewMipFilter);
+	AllocateTempForMips(DownsampleTempData, BaseMip.SizeX, BaseMip.SizeY, FirstMipSizeX, FirstMipSizeY, bDoScaleMipsForAlphaCoverage, KernelDownsample, 2, bSharpenWithoutColorShift, bUnfiltered, bUseNewMipFilter);
 
 	if ( bDownsampleWithAverage )
 	{
-		AllocateTempForMips(AverageTempData, BaseMip.SizeX, BaseMip.SizeY, SecondTempImage.SizeX, SecondTempImage.SizeY, bDoScaleMipsForAlphaCoverage, KernelSimpleAverage, 2, bSharpenWithoutColorShift, bUnfiltered, bUseNewMipFilter);
+		AllocateTempForMips(AverageTempData, BaseMip.SizeX, BaseMip.SizeY, FirstMipSizeX, FirstMipSizeY, bDoScaleMipsForAlphaCoverage, KernelSimpleAverage, 2, bSharpenWithoutColorShift, bUnfiltered, bUseNewMipFilter);
 	}
 
 	// Generate mips
