TextureCompressorModule : two mip gens can run at the same time

#rnx

diff --git a/Engine/Source/Developer/TextureCompressor/Private/TextureCompressorModule.cpp b/Engine/Source/Developer/TextureCompressor/Private/TextureCompressorModule.cpp
--- a/Engine/Source/Developer/TextureCompressor/Private/TextureCompressorModule.cpp
+++ b/Engine/Source/Developer/TextureCompressor/Private/TextureCompressorModule.cpp
@@ -1640,11 +1640,11 @@ static void GenerateMipChain_Old_Reference(
 		IntermediateSrcPtr = &BaseMip;
 
 		// This temp image will be first used as an intermediate destination for the third mip in the chain
-		FirstTempImage.Init( FMath::Max<uint32>( 1, SrcWidth >> 2 ), FMath::Max<uint32>( 1, SrcHeight >> 2 ), Settings.bVolume ? FMath::Max<uint32>( 1, SrcNumSlices >> 2 ) : SrcNumSlices, ImageFormat );
+		FirstTempImage.Init( FMath::Max( 1, SrcWidth >> 2 ), FMath::Max( 1, SrcHeight >> 2 ), Settings.bVolume ? FMath::Max( 1, SrcNumSlices >> 2 ) : SrcNumSlices, ImageFormat );
 	}
 
 	// The image for the first destination
-	FImage SecondTempImage(FMath::Max<uint32>( 1, SrcWidth >> 1 ), FMath::Max<uint32>( 1, SrcHeight >> 1 ), Settings.bVolume ? FMath::Max<uint32>( 1, SrcNumSlices >> 1 ) : SrcNumSlices, ImageFormat);
+	FImage SecondTempImage(FMath::Max( 1, SrcWidth >> 1 ), FMath::Max( 1, SrcHeight >> 1 ), Settings.bVolume ? FMath::Max( 1, SrcNumSlices >> 1 ) : SrcNumSlices, ImageFormat);
 	IntermediateDstPtr = &SecondTempImage;
 
 	// Filtering kernels.
@@ -1809,9 +1809,9 @@ static void GenerateMipChain_Old_Reference(
 		}
 
 		// Update the destination size for the next iteration.
-		IntermediateDstPtr->SizeX = FMath::Max<uint32>( 1, IntermediateSrcPtr->SizeX >> 1 );
-		IntermediateDstPtr->SizeY = FMath::Max<uint32>( 1, IntermediateSrcPtr->SizeY >> 1 );
-		IntermediateDstPtr->NumSlices = Settings.bVolume ? FMath::Max<uint32>( 1, IntermediateSrcPtr->NumSlices >> 1 ) : SrcNumSlices;
+		IntermediateDstPtr->SizeX = FMath::Max( 1, IntermediateSrcPtr->SizeX >> 1 );
+		IntermediateDstPtr->SizeY = FMath::Max( 1, IntermediateSrcPtr->SizeY >> 1 );
+		IntermediateDstPtr->NumSlices = Settings.bVolume ? FMath::Max( 1, IntermediateSrcPtr->NumSlices >> 1 ) : SrcNumSlices;
 	}
 
 	{
@@ -1999,57 +1999,55 @@ void ITextureCompressorModule::GenerateMipChain(
 			}
 			FImageView2D DestView(DestImage, SliceIndex);
 
-			// @@CB 
-			// if ( bDownsampleWithAverage ) these two mipgens can run in parallel I believe?
-			//	and the one to the output mip is not blocking ; <- no, not quite true because we reuse the *source* for the pingpong
-			//	we only need to block on the downsample one
-
-			// DestView is the output mip
-			GenerateSharpenedMipB8G8R8A8(
-				IntermediateSrcView, 
-				IntermediateSrcView2,
-				DownsampleTempData,
-				DestView,
-				AddressMode,
-				bDoScaleMipsForAlphaCoverage,
-				AlphaCoverages,
-				Settings.AlphaCoverageThresholds,
-				KernelDownsample,
-				2,
-				bSharpenWithoutColorShift,
-				bUnfiltered,
-				bUseNewMipFilter,
-				bReDrawBorder);
-
-			// generate IntermediateDstImage:
-			// IntermediateDstImage will be the source for the next mip
-			if ( bDownsampleWithAverage )
-			{
-				// down sample with 2x2 for the next iteration
-				// bDownsampleWithAverage comes from GetMipGenSettings
-				// it is on by default for all cases except Blur
-				// it means every mip is generated *twice*
-				// the output mip is made above using Sharpen from the IntermediateSrc
-				// then the next source is made here using 2x2 ("SimpleAverage")
-				// the next IntermediateSrc is not my outputmip, it's what is made here
-				
-				FImageView2D IntermediateDstView(*IntermediateDstPtr, SliceIndex);
-
+			// lambda to generate one step :
+			auto MakeDestImageMip = [&](FImageView2D & ArgDestView,TArray<FLinearColor> & ArgTempData,FImageKernel2D & ArgKernel) {
+				// DestView is the output mip
 				GenerateSharpenedMipB8G8R8A8(
-					IntermediateSrcView,
+					IntermediateSrcView, 
 					IntermediateSrcView2,
-					AverageTempData,
-					IntermediateDstView,
+					ArgTempData,
+					ArgDestView,
 					AddressMode,
 					bDoScaleMipsForAlphaCoverage,
 					AlphaCoverages,
 					Settings.AlphaCoverageThresholds,
-					KernelSimpleAverage,
+					ArgKernel,
 					2,
 					bSharpenWithoutColorShift,
 					bUnfiltered,
 					bUseNewMipFilter,
-					bReDrawBorder);
+					bReDrawBorder);			
+			};
+
+			// generate IntermediateDstImage:
+			// IntermediateDstImage will be the source for the next mip
+			if ( bDownsampleWithAverage )
+			{
+				// make an async task to generate the output mip:
+				UE::Tasks::TTask<void> MakeDestImageMipTask = UE::Tasks::Launch( TEXT("MakeDestImageMip.Task"), [&](){
+					MakeDestImageMip(DestView,DownsampleTempData,KernelDownsample);
+				} );
+			
+				FImageView2D IntermediateDstView(*IntermediateDstPtr, SliceIndex);
+
+				// down sample with 2x2 for the next iteration
+				//	output of this is used as source for the next level
+				MakeDestImageMip(IntermediateDstView,AverageTempData,KernelSimpleAverage);
+
+				// wait on the task :
+				// we don't depend on the output of this task, so waiting on that is not needed
+				//   but we do use the source image for the pingpong buffer, so you have to wait before reusing that memory
+				// in theory there are ways we could remove this wait (delay it until the whole function returns)
+				//   (one important case is the very first level where Source == BaseImage is not the pingpong)
+				// but then you also have to address the lambda using local variables/lifetimes/etc.
+				MakeDestImageMipTask.GetResult();
+			}
+			else
+			{
+				// single synchronous mipgen per level
+				// generate mips directly in the output chain
+
+				MakeDestImageMip(DestView,DownsampleTempData,KernelDownsample);
 			}
 		}
 
@@ -2086,8 +2084,10 @@ void ITextureCompressorModule::GenerateMipChain(
 		
 		// @@CB this is significant; detach trick again?
 
-		FirstTempImage = FImage();
-		SecondTempImage = FImage();
+		FirstTempImage.RawData.Empty();
+		SecondTempImage.RawData.Empty();		
+		DownsampleTempData.Empty();
+		AverageTempData.Empty();
 	}
 }
 
