TextureCompressorModule: GenerateMipChain detach free and delete test code

#rnx

diff --git a/Engine/Source/Developer/TextureCompressor/Private/TextureCompressorModule.cpp b/Engine/Source/Developer/TextureCompressor/Private/TextureCompressorModule.cpp
--- a/Engine/Source/Developer/TextureCompressor/Private/TextureCompressorModule.cpp
+++ b/Engine/Source/Developer/TextureCompressor/Private/TextureCompressorModule.cpp
@@ -1596,232 +1596,6 @@ static void DownscaleImage(const FImage& SrcImage, FImage& DstImage, const FText
 	}
 }
 
-
-static void GenerateMipChain_Old_Reference(
-	const FTextureBuildSettings& Settings,
-	const FImage& BaseImage,
-	TArray<FImage> &OutMipChain,
-	uint32 MipChainDepth
-	)
-{
-	TRACE_CPUPROFILER_EVENT_SCOPE(Texture.GenerateMipChain);
-
-	// MipChainDepth is the number more to make, OutMipChain has some already
-	// typically BaseImage == OutMipChain.Last()
-	if ( MipChainDepth == 0 )
-	{
-		return;
-	}
-
-	check(BaseImage.Format == ERawImageFormat::RGBA32F);
-
-	const FImage& BaseMip = BaseImage;
-	const int32 SrcWidth = BaseMip.SizeX;
-	const int32 SrcHeight= BaseMip.SizeY;
-	const int32 SrcNumSlices = BaseMip.NumSlices;
-	const ERawImageFormat::Type ImageFormat = ERawImageFormat::RGBA32F;
-
-	const FImage* IntermediateSrcPtr;
-	FImage* IntermediateDstPtr;
-
-	// This will be used as a buffer for the mip processing
-	FImage FirstTempImage;
-
-	if (BaseMip.GammaSpace != EGammaSpace::Linear)
-	{
-		// copy base mip
-		BaseMip.CopyTo(FirstTempImage, ERawImageFormat::RGBA32F, EGammaSpace::Linear);
-
-		IntermediateSrcPtr = &FirstTempImage;
-	}
-	else
-	{
-		// It looks like the BaseMip can be reused for the intermediate source of the second Mip (assuming that the format was check earlier to be RGBA32F)
-		IntermediateSrcPtr = &BaseMip;
-
-		// This temp image will be first used as an intermediate destination for the third mip in the chain
-		FirstTempImage.Init( FMath::Max( 1, SrcWidth >> 2 ), FMath::Max( 1, SrcHeight >> 2 ), Settings.bVolume ? FMath::Max( 1, SrcNumSlices >> 2 ) : SrcNumSlices, ImageFormat );
-	}
-
-	// The image for the first destination
-	FImage SecondTempImage(FMath::Max( 1, SrcWidth >> 1 ), FMath::Max( 1, SrcHeight >> 1 ), Settings.bVolume ? FMath::Max( 1, SrcNumSlices >> 1 ) : SrcNumSlices, ImageFormat);
-	IntermediateDstPtr = &SecondTempImage;
-
-	// Filtering kernels.
-	FImageKernel2D KernelSimpleAverage;
-	FImageKernel2D KernelDownsample;
-	KernelSimpleAverage.BuildSeparatableGaussWithSharpen( 2 );
-	KernelDownsample.BuildSeparatableGaussWithSharpen( Settings.SharpenMipKernelSize, Settings.MipSharpening );
-
-	//@TODO : add a true 3D kernel.
-
-	EMipGenAddressMode AddressMode = ComputeAddressMode(*IntermediateSrcPtr,Settings);
-	bool bReDrawBorder = false;
-	if( Settings.bPreserveBorder )
-	{
-		bReDrawBorder = !Settings.bBorderColorBlack;
-	}
-
-	// Calculate alpha coverage value to preserve along mip chain
-	FVector4f AlphaCoverages(0, 0, 0, 0);
-	if ( Settings.bDoScaleMipsForAlphaCoverage )
-	{
-		check(Settings.AlphaCoverageThresholds != FVector4f(0,0,0,0));
-		check(IntermediateSrcPtr);
-		const FImageView2D IntermediateSrcView = FImageView2D::ConstructConst(*IntermediateSrcPtr, 0);
-
-		const FVector4f AlphaScales(1, 1, 1, 1);		
-		AlphaCoverages = ComputeAlphaCoverage(Settings.AlphaCoverageThresholds, AlphaScales, IntermediateSrcView);
-	}
-
-	TArray<FLinearColor> DownsampleTempData;
-	TArray<FLinearColor> AverageTempData;
-	const bool bDoScaleMipsForAlphaCoverage = Settings.bDoScaleMipsForAlphaCoverage;
-	const bool bSharpenWithoutColorShift = Settings.bSharpenWithoutColorShift;
-	const bool bUnfiltered = Settings.MipGenSettings == TMGS_Unfiltered;
-	const bool bUseNewMipFilter = Settings.bUseNewMipFilter;
-	AllocateTempForMips(DownsampleTempData, BaseMip.SizeX, BaseMip.SizeY, SecondTempImage.SizeX, SecondTempImage.SizeY, bDoScaleMipsForAlphaCoverage, KernelDownsample, 2, bSharpenWithoutColorShift, bUnfiltered, bUseNewMipFilter);
-	AllocateTempForMips(AverageTempData, BaseMip.SizeX, BaseMip.SizeY, SecondTempImage.SizeX, SecondTempImage.SizeY, bDoScaleMipsForAlphaCoverage, KernelSimpleAverage, 2, bSharpenWithoutColorShift, bUnfiltered, bUseNewMipFilter);
-
-	// Generate mips
-	//  default value of MipChainDepth is MAX_uint32, means generate all mips down to 1x1
-	//	(break inside the loop)
-	for (; MipChainDepth != 0 ; --MipChainDepth)
-	{
-		check(IntermediateSrcPtr && IntermediateDstPtr);
-		const FImage& IntermediateSrc = *IntermediateSrcPtr;
-		FImage& IntermediateDst = *IntermediateDstPtr;
-		
-		if ( IntermediateSrc.SizeX == 1 && IntermediateSrc.SizeY == 1 && (!Settings.bVolume || IntermediateSrc.NumSlices == 1))
-		{
-			// should not have been called, starting mip is already small enough
-			check(0);
-			break;
-		}
-
-		// add new mip to TArray<FImage> &OutMipChain :
-		//	placement new on TArray does AddUninitialized then constructs in the last element
-		FImage& DestImage = *new(OutMipChain) FImage(IntermediateDst.SizeX, IntermediateDst.SizeY, IntermediateDst.NumSlices, ImageFormat);
-		
-		for (int32 SliceIndex = 0; SliceIndex < IntermediateDst.NumSlices; ++SliceIndex)
-		{
-			const int32 SrcSliceIndex = Settings.bVolume ? (SliceIndex * 2) : SliceIndex;
-			const FImageView2D IntermediateSrcView = FImageView2D::ConstructConst(IntermediateSrc, SrcSliceIndex);
-			FImageView2D IntermediateSrcView2;
-			if ( Settings.bVolume )
-			{
-				if ( SrcSliceIndex + 1 < IntermediateSrc.NumSlices )
-				{
-					IntermediateSrcView2 = FImageView2D::ConstructConst(IntermediateSrc, SrcSliceIndex + 1);
-				}
-				else
-				{
-					// nonpow2 volume sizeZ , clamp slice index
-					IntermediateSrcView2 = FImageView2D::ConstructConst(IntermediateSrc, SrcSliceIndex);
-				}
-			}
-			FImageView2D DestView(DestImage, SliceIndex);
-			FImageView2D IntermediateDstView(IntermediateDst, SliceIndex);
-
-			// GenerateMipChain_Old_Reference
-
-			// DestView is the output mip
-			GenerateSharpenedMipB8G8R8A8(
-				IntermediateSrcView, 
-				IntermediateSrcView2,
-				DownsampleTempData,
-				DestView,
-				AddressMode,
-				bDoScaleMipsForAlphaCoverage,
-				AlphaCoverages,
-				Settings.AlphaCoverageThresholds,
-				KernelDownsample,
-				2,
-				bSharpenWithoutColorShift,
-				bUnfiltered,
-				bUseNewMipFilter);
-
-			// generate IntermediateDstImage:
-			// IntermediateDstImage will be the source for the next mip
-			if ( Settings.bDownsampleWithAverage )
-			{
-				// down sample without sharpening for the next iteration
-				// bDownsampleWithAverage comes from GetMipGenSettings
-				// it is on by default for all cases except Blur
-				// it means every mip is generated *twice*
-				// the output mip is made above using Sharpen from the IntermediateSrc
-				// then the next source is made here using 2x2 ("SimpleAverage")
-				// the next IntermediateSrc is not my outputmip, it's what is made here
-
-				GenerateSharpenedMipB8G8R8A8(
-					IntermediateSrcView,
-					IntermediateSrcView2,
-					AverageTempData,
-					IntermediateDstView,
-					AddressMode,
-					bDoScaleMipsForAlphaCoverage,
-					AlphaCoverages,
-					Settings.AlphaCoverageThresholds,
-					KernelSimpleAverage,
-					2,
-					bSharpenWithoutColorShift,
-					bUnfiltered,
-					bUseNewMipFilter);
-			}
-		}
-
-		if ( Settings.bDownsampleWithAverage == false )
-		{
-			TRACE_CPUPROFILER_EVENT_SCOPE(Texture.GenerateMipChain.memcpy);
-
-			FMemory::Memcpy( (&IntermediateDst.AsRGBA32F()[0]), (&DestImage.AsRGBA32F()[0]),
-				IntermediateDst.SizeX * IntermediateDst.SizeY * IntermediateDst.NumSlices * sizeof(FLinearColor) );
-		}
-
-		if ( bReDrawBorder )
-		{
-			for (int32 SliceIndex = 0; SliceIndex < IntermediateDst.NumSlices; ++SliceIndex)
-			{
-				const FImageView2D IntermediateSrcView = FImageView2D::ConstructConst(IntermediateSrc, SliceIndex);
-				FImageView2D DestView(DestImage, SliceIndex);
-				FImageView2D IntermediateDstView(IntermediateDst, SliceIndex);
-				GenerateMipBorder( IntermediateSrcView, DestView );
-				GenerateMipBorder( IntermediateSrcView, IntermediateDstView );
-			}
-		}
-
-		// Once we've created mip-maps down to 1x1, we're done.
-		if ( IntermediateDst.SizeX == 1 && IntermediateDst.SizeY == 1 && (!Settings.bVolume || IntermediateDst.NumSlices == 1))
-		{
-			break;
-		}
-
-		// last destination becomes next source
-		if (IntermediateDstPtr == &SecondTempImage)
-		{
-			IntermediateDstPtr = &FirstTempImage;
-			IntermediateSrcPtr = &SecondTempImage;
-		}
-		else
-		{
-			IntermediateDstPtr = &SecondTempImage;
-			IntermediateSrcPtr = &FirstTempImage;
-		}
-
-		// Update the destination size for the next iteration.
-		IntermediateDstPtr->SizeX = FMath::Max( 1, IntermediateSrcPtr->SizeX >> 1 );
-		IntermediateDstPtr->SizeY = FMath::Max( 1, IntermediateSrcPtr->SizeY >> 1 );
-		IntermediateDstPtr->NumSlices = Settings.bVolume ? FMath::Max( 1, IntermediateSrcPtr->NumSlices >> 1 ) : SrcNumSlices;
-	}
-
-	{
-		TRACE_CPUPROFILER_EVENT_SCOPE(Texture.GenerateMipChain.Free);
-		
-		FirstTempImage = FImage();
-		SecondTempImage = FImage();
-	}
-} // GenerateMipChain_Old_Reference
-
 void ITextureCompressorModule::GenerateMipChain(
 	const FTextureBuildSettings& Settings,
 	const FImage& BaseImage,
@@ -2082,12 +1856,36 @@ void ITextureCompressorModule::GenerateMipChain(
 	{
 		TRACE_CPUPROFILER_EVENT_SCOPE(Texture.GenerateMipChain.Free);
 		
-		// @@CB this is significant; detach trick again?
+		// this is significant; detach trick again?
 
+		/*
 		FirstTempImage.RawData.Empty();
 		SecondTempImage.RawData.Empty();		
 		DownsampleTempData.Empty();
 		AverageTempData.Empty();
+		*/
+
+		struct JunkToFree
+		{
+			TArray64<uint8> A;
+			TArray64<uint8> B;
+			TArray<FLinearColor> C;
+			TArray<FLinearColor> D;
+		};
+		
+		JunkToFree * Junk = new JunkToFree;
+		Junk->A = MoveTemp( FirstTempImage.RawData );
+		Junk->B = MoveTemp( SecondTempImage.RawData );
+		Junk->C = MoveTemp( DownsampleTempData );
+		Junk->D = MoveTemp( AverageTempData );
+			
+		UE::Tasks::Launch(TEXT("Texture.GenerateMipChain.Free"),
+			[Junk]() // by value, not ref
+			{
+				TRACE_CPUPROFILER_EVENT_SCOPE(Texture.GenerateMipChain.Free);
+				delete Junk;
+			},
+			LowLevelTasks::ETaskPriority::BackgroundNormal);
 	}
 }
 
@@ -4349,22 +4147,6 @@ private:
 
 				int StartImageIndex = OutMipChain.Num()-1;
 				GenerateMipChain(BuildSettings, OutMipChain[StartImageIndex], OutMipChain, MAX_uint32);
-
-				#if 0 // @@CB temp debug code, remove
-				{
-					// ensure that GenerateMipChain makes the exact same output as GenerateMipChain_Old_Reference :
-					TArray<FImage> ReferenceMipChain;
-					GenerateMipChain_Old_Reference(BuildSettings, OutMipChain[StartImageIndex], ReferenceMipChain, MAX_uint32);
-
-					for(int refi=0,outi=StartImageIndex+1;outi<OutMipChain.Num();outi++,refi++)
-					{
-						FImageInfo OutInfo = OutMipChain[outi];
-						FImageInfo RefInfo = ReferenceMipChain[refi];
-						check( OutInfo == RefInfo );
-						check( OutMipChain[outi].RawData == ReferenceMipChain[refi].RawData ); // runs CompareItems which becomes memcmp
-					}
-				}
-				#endif
 			}
 		}
 		check(OutMipChain.Num() == NumOutputMips);
