Replaced some container operator new usage.

diff --git a/Engine/Source/Developer/ShaderCompilerCommon/Private/ShaderCompilerCommon.cpp b/Engine/Source/Developer/ShaderCompilerCommon/Private/ShaderCompilerCommon.cpp
--- a/Engine/Source/Developer/ShaderCompilerCommon/Private/ShaderCompilerCommon.cpp
+++ b/Engine/Source/Developer/ShaderCompilerCommon/Private/ShaderCompilerCommon.cpp
@@ -1360,8 +1360,8 @@ void CompileOfflineMali(const FShaderCompilerInput& Input, FShaderCompilerOutput
 			{
 				ShaderOutput.bSucceeded = false;
 
-				FShaderCompilerError* NewError = new(ShaderOutput.Errors) FShaderCompilerError();
-				NewError->StrippedErrorMessage = TEXT("[Mali Offline Complier]\n") + StdErr;
+				FShaderCompilerError& NewError = ShaderOutput.Errors.AddDefaulted_GetRef();
+				NewError.StrippedErrorMessage = TEXT("[Mali Offline Complier]\n") + StdErr;
 			}
 			else
 			{
@@ -1369,8 +1369,8 @@ void CompileOfflineMali(const FShaderCompilerInput& Input, FShaderCompilerOutput
 
 				if (Errors.Len())
 				{
-					FShaderCompilerError* NewError = new(ShaderOutput.Errors) FShaderCompilerError();
-					NewError->StrippedErrorMessage = TEXT("[Mali Offline Complier]\n") + Errors;
+					FShaderCompilerError& NewError = ShaderOutput.Errors.AddDefaulted_GetRef();
+					NewError.StrippedErrorMessage = TEXT("[Mali Offline Complier]\n") + Errors;
 					ShaderOutput.bSucceeded = false;
 				}
 			}
@@ -1664,17 +1664,17 @@ namespace CrossCompiler
 	void ParseHlslccError(TArray<FShaderCompilerError>& OutErrors, const FString& InLine, bool bUseAbsolutePaths)
 	{
 		const TCHAR* p = *InLine;
-		FShaderCompilerError* Error = new(OutErrors) FShaderCompilerError();
+		FShaderCompilerError& Error = OutErrors.AddDefaulted_GetRef();
 
 		// Copy the filename.
 		while (*p && *p != TEXT('('))
 		{
-			Error->ErrorVirtualFilePath += (*p++);
+			Error.ErrorVirtualFilePath += (*p++);
 		}
 
 		if (!bUseAbsolutePaths)
 		{
-			Error->ErrorVirtualFilePath = ParseVirtualShaderFilename(Error->ErrorVirtualFilePath);
+			Error.ErrorVirtualFilePath = ParseVirtualShaderFilename(Error.ErrorVirtualFilePath);
 		}
 		p++;
 
@@ -1684,14 +1684,14 @@ namespace CrossCompiler
 		{
 			LineNumber = 10 * LineNumber + (*p++ - TEXT('0'));
 		}
-		Error->ErrorLineString = *FString::Printf(TEXT("%d"), LineNumber);
+		Error.ErrorLineString = *FString::Printf(TEXT("%d"), LineNumber);
 
 		// Skip to the warning message.
 		while (*p && (*p == TEXT(')') || *p == TEXT(':') || *p == TEXT(' ') || *p == TEXT('\t')))
 		{
 			p++;
 		}
-		Error->StrippedErrorMessage = p;
+		Error.StrippedErrorMessage = p;
 	}
 
 
diff --git a/Engine/Source/Developer/ShaderFormatOpenGL/Private/OpenGLShaderCompiler.cpp b/Engine/Source/Developer/ShaderFormatOpenGL/Private/OpenGLShaderCompiler.cpp
--- a/Engine/Source/Developer/ShaderFormatOpenGL/Private/OpenGLShaderCompiler.cpp
+++ b/Engine/Source/Developer/ShaderFormatOpenGL/Private/OpenGLShaderCompiler.cpp
@@ -517,8 +517,8 @@ void ParseGlslError(TArray<FShaderCompilerError>& OutErrors, const FString& InLi
 		{
 			// Note that no mapping exists from the GLSL source to the original
 			// HLSL source.
-			FShaderCompilerError* CompilerError = new(OutErrors) FShaderCompilerError;
-			CompilerError->StrippedErrorMessage = FString::Printf(
+			FShaderCompilerError& CompilerError = OutErrors.AddDefaulted_GetRef();
+			CompilerError.StrippedErrorMessage = FString::Printf(
 				TEXT("driver compile error(%d): %s"),
 				LineNumber,
 				*ErrorMsg
@@ -949,8 +949,8 @@ void FOpenGLFrontend::BuildShaderOutput(
 	if (Header.Bindings.NumSamplers > MaxSamplers)
 	{
 		ShaderOutput.bSucceeded = false;
-		FShaderCompilerError* NewError = new(ShaderOutput.Errors) FShaderCompilerError();
-		NewError->StrippedErrorMessage =
+		FShaderCompilerError& NewError = ShaderOutput.Errors.AddDefaulted_GetRef();
+		NewError.StrippedErrorMessage =
 			FString::Printf(TEXT("shader uses %d samplers exceeding the limit of %d"),
 				Header.Bindings.NumSamplers, MaxSamplers);
 	}
@@ -1029,8 +1029,8 @@ void FOpenGLFrontend::PrecompileShader(FShaderCompilerOutput& ShaderOutput, cons
 	if (GLFrequency == GL_NONE)
 	{
 		ShaderOutput.bSucceeded = false;
-		FShaderCompilerError* NewError = new(ShaderOutput.Errors) FShaderCompilerError();
-		NewError->StrippedErrorMessage = FString::Printf(TEXT("%s shaders not supported for use in OpenGL."), CrossCompiler::GetFrequencyName((EShaderFrequency)ShaderInput.Target.Frequency));
+		FShaderCompilerError& NewError = ShaderOutput.Errors.AddDefaulted_GetRef();
+		NewError.StrippedErrorMessage = FString::Printf(TEXT("%s shaders not supported for use in OpenGL."), CrossCompiler::GetFrequencyName((EShaderFrequency)ShaderInput.Target.Frequency));
 		return;
 	}
 
@@ -1089,8 +1089,8 @@ void FOpenGLFrontend::PrecompileShader(FShaderCompilerOutput& ShaderOutput, cons
 
 				if (ShaderOutput.Errors.Num() == 0)
 				{
-					FShaderCompilerError* NewError = new(ShaderOutput.Errors) FShaderCompilerError();
-					NewError->StrippedErrorMessage = FString::Printf(
+					FShaderCompilerError& NewError = ShaderOutput.Errors.AddDefaulted_GetRef();
+					NewError.StrippedErrorMessage = FString::Printf(
 						TEXT("GLSL source:\n%sGL compile log: %s\n"),
 						ANSI_TO_TCHAR(ShaderSource),
 						ANSI_TO_TCHAR(RawCompileLog.GetData())
@@ -1099,8 +1099,8 @@ void FOpenGLFrontend::PrecompileShader(FShaderCompilerOutput& ShaderOutput, cons
 			}
 			else
 			{
-				FShaderCompilerError* NewError = new(ShaderOutput.Errors) FShaderCompilerError();
-				NewError->StrippedErrorMessage = TEXT("Shader compile failed without errors.");
+				FShaderCompilerError& NewError = ShaderOutput.Errors.AddDefaulted_GetRef();
+				NewError.StrippedErrorMessage = TEXT("Shader compile failed without errors.");
 			}
 
 			ShaderOutput.bSucceeded = false;
@@ -3359,8 +3359,8 @@ void FOpenGLFrontend::CompileShader(const FShaderCompilerInput& Input, FShaderCo
 	if (HlslFrequency == HSF_InvalidFrequency)
 	{
 		Output.bSucceeded = false;
-		FShaderCompilerError* NewError = new(Output.Errors) FShaderCompilerError();
-		NewError->StrippedErrorMessage = FString::Printf(
+		FShaderCompilerError& NewError = Output.Errors.AddDefaulted_GetRef();
+		NewError.StrippedErrorMessage = FString::Printf(
 			TEXT("%s shaders not supported for use in OpenGL."),
 			CrossCompiler::GetFrequencyName(Frequency)
 			);
diff --git a/Engine/Source/Developer/TextureBuild/Private/TextureBuildFunction.cpp b/Engine/Source/Developer/TextureBuild/Private/TextureBuildFunction.cpp
--- a/Engine/Source/Developer/TextureBuild/Private/TextureBuildFunction.cpp
+++ b/Engine/Source/Developer/TextureBuild/Private/TextureBuildFunction.cpp
@@ -238,7 +238,7 @@ static bool TryReadTextureSourceFromCompactBinary(FCbFieldView Source, UE::Deriv
 		int64 MipOffset = MipCbObjectView["Offset"].AsInt64();
 		int64 MipSize = MipCbObjectView["Size"].AsInt64();
 
-		FImage* SourceMip = new(OutMips) FImage(
+		FImage& SourceMip = OutMips.Emplace_GetRef(
 			MipSizeX, MipSizeY,
 			NumSlices,
 			RawImageFormat,
@@ -246,20 +246,20 @@ static bool TryReadTextureSourceFromCompactBinary(FCbFieldView Source, UE::Deriv
 		);
 
 		check( MipOffset + MipSize <= DecompressedSourceDataSize );
-		check( SourceMip->GetImageSizeBytes() == MipSize );
+		check( SourceMip.GetImageSizeBytes() == MipSize );
 
 		if ((MipsCbArrayView.Num() == 1) && (CompressionFormat != TSCF_None))
 		{
 			// In the case where there is only one mip and its already in a TArray, there is no need to allocate new array contents, just use a move instead
 			check( MipOffset == 0 );
-			SourceMip->RawData = MoveTemp(IntermediateDecompressedData);
+			SourceMip.RawData = MoveTemp(IntermediateDecompressedData);
 		}
 		else
 		{
-			SourceMip->RawData.Reset(MipSize);
-			SourceMip->RawData.AddUninitialized(MipSize);
+			SourceMip.RawData.Reset(MipSize);
+			SourceMip.RawData.AddUninitialized(MipSize);
 			FMemory::Memcpy(
-				SourceMip->RawData.GetData(),
+				SourceMip.RawData.GetData(),
 				DecompressedSourceData + MipOffset,
 				MipSize
 			);
diff --git a/Engine/Source/Developer/TextureCompressor/Private/TextureCompressorModule.cpp b/Engine/Source/Developer/TextureCompressor/Private/TextureCompressorModule.cpp
--- a/Engine/Source/Developer/TextureCompressor/Private/TextureCompressorModule.cpp
+++ b/Engine/Source/Developer/TextureCompressor/Private/TextureCompressorModule.cpp
@@ -1759,7 +1759,7 @@ void ITextureCompressorModule::GenerateMipChain(
 		// add new mip to TArray<FImage> &OutMipChain :
 		//	placement new on TArray does AddUninitialized then constructs in the last element
 		check( OutMipChain.GetSlack() > 0 ); // should have been reserved
-		FImage& DestImage = *new(OutMipChain) FImage(DstSizeX, DstSizeY, DstNumSlices, ImageFormat);
+		FImage& DestImage = OutMipChain.Emplace_GetRef(DstSizeX, DstSizeY, DstNumSlices, ImageFormat);
 		
 		for (int32 SliceIndex = 0; SliceIndex < DstNumSlices; ++SliceIndex)
 		{
@@ -2381,12 +2381,12 @@ void ITextureCompressorModule::GenerateAngularFilteredMips(TArray<FImage>& InOut
 		FImage& BaseMip = SrcMipChain[MipIndex - 1];
 		int32 BaseExtent = BaseMip.SizeX;
 		int32 MipExtent = FMath::Max(BaseExtent >> 1, 1);
-		FImage* Mip = new(SrcMipChain) FImage(MipExtent, MipExtent, BaseMip.NumSlices, BaseMip.Format);
+		FImage& Mip = SrcMipChain.Emplace_GetRef(MipExtent, MipExtent, BaseMip.NumSlices, BaseMip.Format);
 
 		for(int32 Face = 0; Face < 6; ++Face)
 		{
 			FImageView2D BaseMipView(BaseMip, Face);
-			FImageView2D MipView(*Mip, Face);
+			FImageView2D MipView(Mip, Face);
 
 			for(int32 y = 0; y < MipExtent; ++y)
 			{
@@ -2441,8 +2441,8 @@ void ITextureCompressorModule::GenerateAngularFilteredMips(TArray<FImage>& InOut
 		float FloatInputMip = 0.5f * FMath::Log2(AreaCoveredInNormalizedArea) + NumMips - QualityBias;
 		uint32 InputMip = FMath::Clamp(FMath::TruncToInt(FloatInputMip), 0, NumMips - 1);
 
-		FImage* Mip = new(InOutMipChain) FImage(Extent, Extent, 6, ERawImageFormat::RGBA32F);
-		GenerateAngularFilteredMip(Mip, SrcMipChain[InputMip], ConeAngle);
+		FImage& Mip = InOutMipChain.Emplace_GetRef(Extent, Extent, 6, ERawImageFormat::RGBA32F);
+		GenerateAngularFilteredMip(&Mip, SrcMipChain[InputMip], ConeAngle);
 		Extent = FMath::Max(Extent >> 1, 1);
 	}
 }
@@ -3841,7 +3841,7 @@ private:
 
 				// space for one source mip and one destination mip
 				const FImage& SourceImage = bSuitableFormat ? FirstSourceMipImage : Temp;
-				FImage& TargetImage = *new (PaddedSourceMips) FImage(TargetTextureSizeX, TargetTextureSizeY, BuildSettings.bVolume ? TargetTextureSizeZ : SourceImage.NumSlices, SourceImage.Format);
+				FImage& TargetImage = PaddedSourceMips.Emplace_GetRef(TargetTextureSizeX, TargetTextureSizeY, BuildSettings.bVolume ? TargetTextureSizeZ : SourceImage.NumSlices, SourceImage.Format);
 				FLinearColor FillColor = BuildSettings.PaddingColor;
 
 				FLinearColor* TargetPtr = (FLinearColor*)TargetImage.RawData.GetData();
@@ -4033,12 +4033,12 @@ private:
 			// this is a code dupe of the processing done in GenerateMipChain
 
 			// create base for the mip chain
-			FImage* Mip = new(OutMipChain) FImage();
+			FImage& Mip = OutMipChain.AddDefaulted_GetRef();
 
 			if (bLongLatCubemap)
 			{
 				// Generate the base mip from the long-lat source image.
-				GenerateBaseCubeMipFromLongitudeLatitude2D(Mip, Image, BuildSettings.MaxTextureResolution, BuildSettings.SourceEncodingOverride);
+				GenerateBaseCubeMipFromLongitudeLatitude2D(&Mip, Image, BuildSettings.MaxTextureResolution, BuildSettings.SourceEncodingOverride);
 	
 				check( CopyCount == 1 );
 			}
@@ -4054,16 +4054,16 @@ private:
 						NormalizeMip(Temp);
 					}
 
-					GenerateTopMip(Temp, *Mip, BuildSettings);
+					GenerateTopMip(Temp, Mip, BuildSettings);
 				}
 				else
 				{
 					if (bLinearize)
 					{
-						Image.Linearize(BuildSettings.SourceEncodingOverride, *Mip);
+						Image.Linearize(BuildSettings.SourceEncodingOverride, Mip);
 						if (BuildSettings.bRenormalizeTopMip)
 						{
-							NormalizeMip(*Mip);
+							NormalizeMip(Mip);
 						}
 					}
 					else
@@ -4077,7 +4077,7 @@ private:
 							DestFormat = ERawImageFormat::RGBA32F;
 							DestGammaSpace = EGammaSpace::Linear;
 						}
-						Image.CopyTo(*Mip, DestFormat, DestGammaSpace);
+						Image.CopyTo(Mip, DestFormat, DestGammaSpace);
 
 						//@@CB todo : when Mip format == Image format, we can Move instead of Copy
 						//	have to make sure that's okay with SourceMips/TextureData
@@ -4089,13 +4089,13 @@ private:
 
 			if (BuildSettings.Downscale > 1.f)
 			{		
-				DownscaleImage(*Mip, *Mip, FTextureDownscaleSettings(BuildSettings));
+				DownscaleImage(Mip, Mip, FTextureDownscaleSettings(BuildSettings));
 			}
 
 			if (BuildSettings.bHasColorSpaceDefinition)
 			{
 				FImageCore::TransformToWorkingColorSpace(
-					*Mip,
+					Mip,
 					FVector2d(BuildSettings.RedChromaticityCoordinate),
 					FVector2d(BuildSettings.GreenChromaticityCoordinate),
 					FVector2d(BuildSettings.BlueChromaticityCoordinate),
@@ -4104,16 +4104,16 @@ private:
 			}
 
 			// Apply color adjustments
-			AdjustImageColors(*Mip, BuildSettings);
+			AdjustImageColors(Mip, BuildSettings);
 
 			if (BuildSettings.bComputeBokehAlpha)
 			{
 				// To get the occlusion in the BokehDOF shader working for all Bokeh textures.
-				ComputeBokehAlpha(*Mip);
+				ComputeBokehAlpha(Mip);
 			}
 			if (BuildSettings.bFlipGreenChannel)
 			{
-				FlipGreenChannel(*Mip);
+				FlipGreenChannel(Mip);
 			}
 		}
 
diff --git a/Engine/Source/Developer/VulkanShaderFormat/Private/VulkanShaderCompiler.cpp b/Engine/Source/Developer/VulkanShaderFormat/Private/VulkanShaderCompiler.cpp
--- a/Engine/Source/Developer/VulkanShaderFormat/Private/VulkanShaderCompiler.cpp
+++ b/Engine/Source/Developer/VulkanShaderFormat/Private/VulkanShaderCompiler.cpp
@@ -1063,9 +1063,9 @@ static void BuildShaderOutput(
 			break;
 		case EVulkanBindingType::PackedUniformBuffer:
 			{
-				FOLDVulkanCodeHeader::FPackedUBToVulkanBindingIndex* New = new(OLDHeader.NEWPackedUBToVulkanBindingIndices) FOLDVulkanCodeHeader::FPackedUBToVulkanBindingIndex;
-				New->TypeName = (CrossCompiler::EPackedTypeName)Binding.SubType;
-				New->VulkanBindingIndex = StageOffset + Index;
+				FOLDVulkanCodeHeader::FPackedUBToVulkanBindingIndex& New = OLDHeader.NEWPackedUBToVulkanBindingIndices.AddDefaulted_GetRef();
+				New.TypeName = (CrossCompiler::EPackedTypeName)Binding.SubType;
+				New.VulkanBindingIndex = StageOffset + Index;
 				++OLDHeader.NEWDescriptorInfo.NumBufferInfos;
 			}
 			break;
@@ -2572,8 +2572,8 @@ void DoCompileVulkanShader(const FShaderCompilerInput& Input, FShaderCompilerOut
 	if (HlslFrequency == HSF_InvalidFrequency)
 	{
 		Output.bSucceeded = false;
-		FShaderCompilerError* NewError = new(Output.Errors) FShaderCompilerError();
-		NewError->StrippedErrorMessage = FString::Printf(
+		FShaderCompilerError& NewError = Output.Errors.AddDefaulted_GetRef();
+		NewError.StrippedErrorMessage = FString::Printf(
 			TEXT("%s shaders not supported for use in Vulkan."),
 			CrossCompiler::GetFrequencyName(Frequency));
 		return;
diff --git a/Engine/Source/Programs/ShaderCompileWorker/Private/ShaderCompileWorker.cpp b/Engine/Source/Programs/ShaderCompileWorker/Private/ShaderCompileWorker.cpp
--- a/Engine/Source/Programs/ShaderCompileWorker/Private/ShaderCompileWorker.cpp
+++ b/Engine/Source/Programs/ShaderCompileWorker/Private/ShaderCompileWorker.cpp
@@ -603,7 +603,7 @@ private:
 			
 			for (int32 BatchIndex = 0; BatchIndex < NumBatches; BatchIndex++)
 			{
-				FShaderCompileJob& Job = *new(OutSingleJobs) FShaderCompileJob;
+				FShaderCompileJob& Job = OutSingleJobs.AddDefaulted_GetRef();
 				// Deserialize the job's inputs.
 				Job.SerializeWorkerInput(InputFile);
 				Job.Input.DeserializeSharedInputs(InputFile, ExternalIncludes, SharedEnvironments, ParameterStructures);
@@ -638,17 +638,17 @@ private:
 
 			for (int32 Index = 0; Index < NumPipelines; ++Index)
 			{
-				FString& PipelineName = *new(OutPipelineNames) FString();
+				FString& PipelineName = OutPipelineNames.AddDefaulted_GetRef();
 				InputFile << PipelineName;
 				
 				int32 NumStages = 0;
 				InputFile << NumStages;
-				FShaderPipelineCompileJob* PipelineJob = new(OutPipelineJobs) FShaderPipelineCompileJob(NumStages);
+				FShaderPipelineCompileJob& PipelineJob = OutPipelineJobs.Emplace_GetRef(NumStages);
 
 				for (int32 StageIndex = 0; StageIndex < NumStages; ++StageIndex)
 				{
 					// Deserialize the job's inputs.
-					FShaderCompileJob* Job = PipelineJob->StageJobs[StageIndex]->GetSingleShaderJob();
+					FShaderCompileJob* Job = PipelineJob.StageJobs[StageIndex]->GetSingleShaderJob();
 					Job->SerializeWorkerInput(InputFile);
 					Job->Input.DeserializeSharedInputs(InputFile, ExternalIncludes, SharedEnvironments, ParameterStructures);
 
@@ -661,7 +661,7 @@ private:
 					}
 				}
 
-				CompileShaderPipeline(GetShaderFormats(), PipelineJob, WorkingDirectory, &GNumProcessedJobs);
+				CompileShaderPipeline(GetShaderFormats(), &PipelineJob, WorkingDirectory, &GNumProcessedJobs);
 			}
 		}
 	}
diff --git a/Engine/Source/Runtime/Core/Private/Stats/Stats2.cpp b/Engine/Source/Runtime/Core/Private/Stats/Stats2.cpp
--- a/Engine/Source/Runtime/Core/Private/Stats/Stats2.cpp
+++ b/Engine/Source/Runtime/Core/Private/Stats/Stats2.cpp
@@ -331,8 +331,8 @@ void FStartupMessages::AddMetadata( FName InStatName, const TCHAR* InStatDesc, c
 	LLM_SCOPE(ELLMTag::Stats);
 	FScopeLock Lock( &CriticalSection );
 
-	new (DelayedMessages)FStatMessage( InGroupName, EStatDataType::ST_None, "Groups", InGroupCategory, InGroupDesc, false, false, bSortByName );
-	new (DelayedMessages)FStatMessage( InStatName, InStatType, InGroupName, InGroupCategory, InStatDesc, bShouldClearEveryFrame, bCycleStat, bSortByName, InMemoryRegion );
+	DelayedMessages.Emplace( InGroupName, EStatDataType::ST_None, "Groups", InGroupCategory, InGroupDesc, false, false, bSortByName );
+	DelayedMessages.Emplace( InStatName, InStatType, InGroupName, InGroupCategory, InStatDesc, bShouldClearEveryFrame, bCycleStat, bSortByName, InMemoryRegion );
 }
 
 
diff --git a/Engine/Source/Runtime/Core/Private/Stats/StatsCommand.cpp b/Engine/Source/Runtime/Core/Private/Stats/StatsCommand.cpp
--- a/Engine/Source/Runtime/Core/Private/Stats/StatsCommand.cpp
+++ b/Engine/Source/Runtime/Core/Private/Stats/StatsCommand.cpp
@@ -1397,7 +1397,7 @@ struct FHUDGroupManager
 			for( int32 Index = 0; Index < TotalAggregateInclusive.Num(); ++Index )
 			{
 				const FStatMessage& StatMessage = TotalAggregateInclusive[Index];
-				new(AggregatedFlatHistory) FComplexStatMessage(StatMessage);
+				AggregatedFlatHistory.Emplace(StatMessage);
 			}
 
 			// Copy flat-stack stats by thread
@@ -1407,7 +1407,7 @@ struct FHUDGroupManager
 				TArray<FComplexStatMessage>& AggregatedFlatHistoryThreadBreakdownArray = AggregatedFlatHistoryThreadBreakdown.Add(It.Key());
 				for (const FStatMessage& StatMessage : It.Value())
 				{
-					new (AggregatedFlatHistoryThreadBreakdownArray)FComplexStatMessage(StatMessage);
+					AggregatedFlatHistoryThreadBreakdownArray.Emplace(StatMessage);
 				}
 			}
 
@@ -1416,7 +1416,7 @@ struct FHUDGroupManager
 			for( int32 Index = 0; Index < TotalNonStackStats.Num(); ++Index )
 			{
 				const FStatMessage& StatMessage = TotalNonStackStats[Index];
-				new(AggregatedNonStackStatsHistory) FComplexStatMessage(StatMessage);
+				AggregatedNonStackStatsHistory.Emplace(StatMessage);
 			}
 			
 			// Accumulate hierarchy, flat and non-stack stats.
@@ -1498,7 +1498,7 @@ struct FHUDGroupManager
 						const bool bToBeAdded = InternalGroup.EnabledItems.Contains( AggregatedStatMessage.NameAndInfo.GetRawName() );
 						if( bToBeAdded )
 						{
-							new(HudGroup.FlatAggregate) FComplexStatMessage( AggregatedStatMessage );
+							HudGroup.FlatAggregate.Add( AggregatedStatMessage );
 						}
 					}
 
@@ -1512,7 +1512,7 @@ struct FHUDGroupManager
 							if(bToBeAdded)
 							{
 								TArray<FComplexStatMessage>& DestArray = HudGroup.FlatAggregateThreadBreakdown.FindOrAdd(It.Key());	
-								new(DestArray) FComplexStatMessage( AggregatedStatMessage );
+								DestArray.Add( AggregatedStatMessage );
 							}
 						}
 					}
@@ -1528,8 +1528,8 @@ struct FHUDGroupManager
 					const bool bToBeAdded = InternalGroup.EnabledItems.Contains( AggregatedStatMessage.NameAndInfo.GetRawName() );
 					if( bToBeAdded )
 					{
-						new(Dest) FComplexStatMessage(AggregatedStatMessage);
-					}	
+						Dest.Add(AggregatedStatMessage);
+					}
 				}
 			}
 
diff --git a/Engine/Source/Runtime/Core/Private/Stats/StatsData.cpp b/Engine/Source/Runtime/Core/Private/Stats/StatsData.cpp
--- a/Engine/Source/Runtime/Core/Private/Stats/StatsData.cpp
+++ b/Engine/Source/Runtime/Core/Private/Stats/StatsData.cpp
@@ -202,7 +202,7 @@ void FRawStatStackNode::AddNameHierarchy(int32 CurrentPrefixDepth)
 			for (int32 Index = 0; Index < ChildArray.Num(); Index++)
 			{
 				FRawStatStackNode& Child = *ChildArray[Index];
-				new (ChildNames) TArray<FName>();
+				ChildNames.AddDefaulted();
 				TArray<FName>& ParsedNames = ChildNames[Index];
 
 				TArray<FString> Parts;
@@ -210,7 +210,7 @@ void FRawStatStackNode::AddNameHierarchy(int32 CurrentPrefixDepth)
 				if (Name.StartsWith(TEXT("//")))
 				{
 					// we won't add hierarchy for grouped stats
-					new (ParsedNames) FName(Child.Meta.NameAndInfo.GetRawName());
+					ParsedNames.Add(Child.Meta.NameAndInfo.GetRawName());
 				}
 				else
 				{
@@ -220,7 +220,7 @@ void FRawStatStackNode::AddNameHierarchy(int32 CurrentPrefixDepth)
 					ParsedNames.Empty(Parts.Num());
 					for (int32 PartIndex = 0; PartIndex < Parts.Num(); PartIndex++)
 					{
-						new (ParsedNames) FName(*Parts[PartIndex]);
+						ParsedNames.Add(*Parts[PartIndex]);
 					}
 				}
 			}
@@ -427,10 +427,10 @@ void FRawStatStackNode::DebugPrintLeafFilterInner(TCHAR const* Filter, int32 Dep
 
 void FRawStatStackNode::Encode(TArray<FStatMessage>& OutStats) const
 {
-	FStatMessage* NewStat = new (OutStats) FStatMessage(Meta);
+	FStatMessage& NewStat = OutStats.Add_GetRef(Meta);
 	if (Children.Num())
 	{
-		NewStat->NameAndInfo.SetField<EStatOperation>(EStatOperation::ChildrenStart);
+		NewStat.NameAndInfo.SetField<EStatOperation>(EStatOperation::ChildrenStart);
 		for (TMap<FName, FRawStatStackNode*>::TConstIterator It(Children); It; ++It)
 		{
 			FRawStatStackNode const* Child = It.Value();
@@ -441,7 +441,7 @@ void FRawStatStackNode::Encode(TArray<FStatMessage>& OutStats) const
 	}
 	else
 	{
-		NewStat->NameAndInfo.SetField<EStatOperation>(EStatOperation::Leaf);
+		NewStat.NameAndInfo.SetField<EStatOperation>(EStatOperation::Leaf);
 	}
 }
 
@@ -1547,7 +1547,7 @@ void FStatsThreadState::GetRawStackStats(int64 TargetFrame, FRawStatStackNode& R
 	{
 		for (TMap<FName, FStatMessage>::TConstIterator It(ThisFrameNonStackStats); It; ++It)
 		{
-			new (*OutNonStackStats) FStatMessage(It.Value());
+			OutNonStackStats->Add(It.Value());
 		}
 	}
 }
@@ -1712,8 +1712,8 @@ FName FStatsThreadState::GetStatThreadName( const FStatPacket& Packet ) const
 
 void FStatsThreadState::Condense(int64 TargetFrame, TArray<FStatMessage>& OutStats) const
 {
-	new (OutStats) FStatMessage(FStatConstants::AdvanceFrame.GetEncodedName(), EStatOperation::AdvanceFrameEventGameThread, TargetFrame, false);
-	new (OutStats) FStatMessage(FStatConstants::AdvanceFrame.GetEncodedName(), EStatOperation::AdvanceFrameEventRenderThread, TargetFrame, false);
+	OutStats.Emplace(FStatConstants::AdvanceFrame.GetEncodedName(), EStatOperation::AdvanceFrameEventGameThread, TargetFrame, false);
+	OutStats.Emplace(FStatConstants::AdvanceFrame.GetEncodedName(), EStatOperation::AdvanceFrameEventRenderThread, TargetFrame, false);
 	FRawStatStackNode Root;
 	GetRawStackStats(TargetFrame, Root, &OutStats);
 	TArray<FStatMessage> StackStats;
@@ -1793,7 +1793,7 @@ void FStatsThreadState::AddMissingStats(TArray<FStatMessage>& Dest, TSet<FName>
 		FStatMessage const* Zero = ShortNameToLongName.Find(*It);
 		if (Zero)
 		{
-			new (Dest) FStatMessage(*Zero);
+			Dest.Add(*Zero);
 		}
 	}
 }
diff --git a/Engine/Source/Runtime/Core/Private/Stats/StatsFile.cpp b/Engine/Source/Runtime/Core/Private/Stats/StatsFile.cpp
--- a/Engine/Source/Runtime/Core/Private/Stats/StatsFile.cpp
+++ b/Engine/Source/Runtime/Core/Private/Stats/StatsFile.cpp
@@ -818,7 +818,7 @@ void FStatsReadFile::ReadRegularStats()
 		{
 			// Read the message.
 			FStatMessage Message( Stream.ReadMessage( MemoryReader, bIsFinalized ) );
-			new (PendingMessages)FStatMessage( Message );
+			PendingMessages.Add( MoveTemp( Message ) );
 		}
 
 		// Adds a new from from the pending messages, the pending messages will be removed after the call.
diff --git a/Engine/Source/Runtime/CoreUObject/Private/AssetRegistry/AssetBundleData.cpp b/Engine/Source/Runtime/CoreUObject/Private/AssetRegistry/AssetBundleData.cpp
--- a/Engine/Source/Runtime/CoreUObject/Private/AssetRegistry/AssetBundleData.cpp
+++ b/Engine/Source/Runtime/CoreUObject/Private/AssetRegistry/AssetBundleData.cpp
@@ -282,7 +282,7 @@ void FAssetBundleData::AddBundleAsset(FName BundleName, const FTopLevelAssetPath
 
 	if (!FoundEntry)
 	{
-		FoundEntry = new(Bundles) FAssetBundleEntry(BundleName);
+		FoundEntry = &Bundles.Emplace_GetRef(BundleName);
 	}
 
 #if WITH_EDITORONLY_DATA
@@ -304,7 +304,7 @@ void FAssetBundleData::AddBundleAssets(FName BundleName, const TArray<FTopLevelA
 			// Only create if required
 			if (!FoundEntry)
 			{
-				FoundEntry = new(Bundles) FAssetBundleEntry(BundleName);
+				FoundEntry = &Bundles.Emplace_GetRef(BundleName);
 			}
 
 #if WITH_EDITORONLY_DATA
@@ -323,7 +323,7 @@ void FAssetBundleData::SetBundleAssets(FName BundleName, TArray<FTopLevelAssetPa
 
 	if (!FoundEntry)
 	{
-		FoundEntry = new(Bundles) FAssetBundleEntry(BundleName);
+		FoundEntry = &Bundles.Emplace_GetRef(BundleName);
 	}
 
 	FoundEntry->AssetPaths = MoveTemp(AssetPaths);
@@ -346,7 +346,7 @@ void FAssetBundleData::AddBundleAsset(FName BundleName, const FSoftObjectPath& A
 
 	if (!FoundEntry)
 	{
-		FoundEntry = new(Bundles) FAssetBundleEntry(BundleName);
+		FoundEntry = &Bundles.Emplace_GetRef(BundleName);
 	}
 
 #if WITH_EDITORONLY_DATA
@@ -368,7 +368,7 @@ void FAssetBundleData::AddBundleAssets(FName BundleName, const TArray<FSoftObjec
 			// Only create if required
 			if (!FoundEntry)
 			{
-				FoundEntry = new(Bundles) FAssetBundleEntry(BundleName);
+				FoundEntry = &Bundles.Emplace_GetRef(BundleName);
 			}
 
 #if WITH_EDITORONLY_DATA
@@ -387,7 +387,7 @@ void FAssetBundleData::SetBundleAssets(FName BundleName, TArray<FSoftObjectPath>
 
 	if (!FoundEntry)
 	{
-		FoundEntry = new(Bundles) FAssetBundleEntry(BundleName);
+		FoundEntry = &Bundles.Emplace_GetRef(BundleName);
 	}
 
 	FoundEntry->AssetPaths.Reset(AssetPaths.Num());
@@ -415,7 +415,7 @@ void FAssetBundleData::AddBundleAssetTruncated(FName BundleName, const FSoftObje
 
 	if (!FoundEntry)
 	{
-		FoundEntry = new(Bundles) FAssetBundleEntry(BundleName);
+		FoundEntry = &Bundles.Emplace_GetRef(BundleName);
 	}
 
 #if WITH_EDITORONLY_DATA
@@ -437,7 +437,7 @@ void FAssetBundleData::AddBundleAssetsTruncated(FName BundleName, const TArray<F
 			// Only create if required
 			if (!FoundEntry)
 			{
-				FoundEntry = new(Bundles) FAssetBundleEntry(BundleName);
+				FoundEntry = &Bundles.Emplace_GetRef(BundleName);
 			}
 
 #if WITH_EDITORONLY_DATA
@@ -456,7 +456,7 @@ void FAssetBundleData::SetBundleAssetsTruncated(FName BundleName, const TArray<F
 
 	if (!FoundEntry)
 	{
-		FoundEntry = new(Bundles) FAssetBundleEntry(BundleName);
+		FoundEntry = &Bundles.Emplace_GetRef(BundleName);
 	}
 
 	FoundEntry->AssetPaths.Reset(AssetPaths.Num());
diff --git a/Engine/Source/Runtime/CoreUObject/Private/UObject/Class.cpp b/Engine/Source/Runtime/CoreUObject/Private/UObject/Class.cpp
--- a/Engine/Source/Runtime/CoreUObject/Private/UObject/Class.cpp
+++ b/Engine/Source/Runtime/CoreUObject/Private/UObject/Class.cpp
@@ -4584,7 +4584,7 @@ void UClass::SetUpRuntimeReplicationData()
 			NetProperties[i]->RepIndex = (uint16)ClassReps.Num();
 			for (int32 j = 0; j < NetProperties[i]->ArrayDim; j++)
 			{
-				new(ClassReps)FRepRecord(NetProperties[i], j);
+				ClassReps.Emplace(NetProperties[i], j);
 			}
 		}
 		check(ClassReps.Num() <= 65535);
@@ -5829,7 +5829,7 @@ void UClass::AddNativeFunction(const ANSICHAR* InName, FNativeFuncPtr InPointer)
 		}
 	}
 #endif
-	new(NativeFunctionLookupTable) FNativeFunctionLookup(InFName,InPointer);
+	NativeFunctionLookupTable.Emplace(InFName,InPointer);
 }
 
 void UClass::AddNativeFunction(const WIDECHAR* InName, FNativeFuncPtr InPointer)
@@ -5850,7 +5850,7 @@ void UClass::AddNativeFunction(const WIDECHAR* InName, FNativeFuncPtr InPointer)
 		}
 	}
 #endif
-	new(NativeFunctionLookupTable)FNativeFunctionLookup(InFName, InPointer);
+	NativeFunctionLookupTable.Emplace(InFName, InPointer);
 }
 
 void UClass::CreateLinkAndAddChildFunctionsToMap(const FClassFunctionLinkInfo* Functions, uint32 NumFunctions)
diff --git a/Engine/Source/Runtime/CoreUObject/Private/UObject/LinkerLoad.cpp b/Engine/Source/Runtime/CoreUObject/Private/UObject/LinkerLoad.cpp
--- a/Engine/Source/Runtime/CoreUObject/Private/UObject/LinkerLoad.cpp
+++ b/Engine/Source/Runtime/CoreUObject/Private/UObject/LinkerLoad.cpp
@@ -1814,8 +1814,8 @@ FLinkerLoad::ELinkerStatus FLinkerLoad::SerializeGatherableTextDataMap(bool bFor
 	FStructuredArchive::FStream Stream = StructuredArchiveRootRecord->EnterStream(TEXT("GatherableTextData"));
 	while (GatherableTextDataMapIndex < Summary.GatherableTextDataCount && !IsTimeLimitExceeded(TEXT("serializing gatherable text data map"), 100))
 	{
-		FGatherableTextData* GatherableTextData = new(GatherableTextDataMap)FGatherableTextData;
-		Stream.EnterElement() << *GatherableTextData;
+		FGatherableTextData& GatherableTextData = GatherableTextDataMap.AddDefaulted_GetRef();
+		Stream.EnterElement() << GatherableTextData;
 		GatherableTextDataMapIndex++;
 	}
 
@@ -1841,8 +1841,8 @@ FLinkerLoad::ELinkerStatus FLinkerLoad::SerializeImportMap()
 
 	while( ImportMapIndex < Summary.ImportCount && !IsTimeLimitExceeded(TEXT("serializing import map"),100) )
 	{
-		FObjectImport* Import = new(ImportMap)FObjectImport;
-		Stream.EnterElement() << *Import;
+		FObjectImport& Import = ImportMap.AddDefaulted_GetRef();
+		Stream.EnterElement() << Import;
 		ImportMapIndex++;
 	}
 	
@@ -2006,15 +2006,15 @@ FLinkerLoad::ELinkerStatus FLinkerLoad::FixupImportMap()
 				for (FName NewPackage : NewPackageImports)
 				{
 					// We are adding a new import to the map as we need the new package dependency added to the works
-					FObjectImport* NewImport = new (ImportMap) FObjectImport();
-
-					NewImport->ClassName = NAME_Package;
-					NewImport->ClassPackage = GLongCoreUObjectPackageName;
-					NewImport->ObjectName = NewPackage;
-					NewImport->OuterIndex = FPackageIndex();
-					NewImport->XObject = 0;
-					NewImport->SourceLinker = 0;
-					NewImport->SourceIndex = -1;
+					FObjectImport& NewImport = ImportMap.AddDefaulted_GetRef();
+
+					NewImport.ClassName = NAME_Package;
+					NewImport.ClassPackage = GLongCoreUObjectPackageName;
+					NewImport.ObjectName = NewPackage;
+					NewImport.OuterIndex = FPackageIndex();
+					NewImport.XObject = 0;
+					NewImport.SourceLinker = 0;
+					NewImport.SourceIndex = -1;
 				}
 			}
 
@@ -2164,10 +2164,10 @@ FLinkerLoad::ELinkerStatus FLinkerLoad::SerializeExportMap()
 
 	while( ExportMapIndex < Summary.ExportCount && !IsTimeLimitExceeded(TEXT("serializing export map"),100) )
 	{
-		FObjectExport* Export = new(ExportMap)FObjectExport;
-		Stream.EnterElement() << *Export;
-		Export->ThisIndex = FPackageIndex::FromExport(ExportMapIndex);
-		Export->bWasFiltered = FilterExport(*Export);
+		FObjectExport& Export = ExportMap.AddDefaulted_GetRef();
+		Stream.EnterElement() << Export;
+		Export.ThisIndex = FPackageIndex::FromExport(ExportMapIndex);
+		Export.bWasFiltered = FilterExport(Export);
 		ExportMapIndex++;
 	}
 
@@ -5082,7 +5082,7 @@ UObject* FLinkerLoad::CreateExport( int32 Index )
 				// A redirector has been found, replace this export with it.
 				LoadClass = UObjectRedirector::StaticClass();
 				// Create new import for UObjectRedirector class
-				FObjectImport* RedirectorImport = new(ImportMap)FObjectImport(UObjectRedirector::StaticClass());
+				FObjectImport& RedirectorImport = ImportMap.Emplace_GetRef(UObjectRedirector::StaticClass());
 				check(CurrentLoadContext);
 				CurrentLoadContext->IncrementImportCount();
 				FLinkerManager::Get().AddLoaderWithNewImports(this);				
diff --git a/Engine/Source/Runtime/CoreUObject/Private/UObject/Obj.cpp b/Engine/Source/Runtime/CoreUObject/Private/UObject/Obj.cpp
--- a/Engine/Source/Runtime/CoreUObject/Private/UObject/Obj.cpp
+++ b/Engine/Source/Runtime/CoreUObject/Private/UObject/Obj.cpp
@@ -3510,7 +3510,7 @@ void UObject::RetrieveReferencers( TArray<FReferencerInformation>* OutInternalRe
 				{
 					// manually allocate just one element - much slower but avoids slack which improves success rate on consoles
 					OutInternalReferencers->Reserve(OutInternalReferencers->Num() + 1);
-					new(*OutInternalReferencers) FReferencerInformation(Object, Count, Referencers);
+					OutInternalReferencers->Emplace(Object, Count, Referencers);
 				}
 			}
 			else
@@ -3519,7 +3519,7 @@ void UObject::RetrieveReferencers( TArray<FReferencerInformation>* OutInternalRe
 				{
 					// manually allocate just one element - much slower but avoids slack which improves success rate on consoles
 					OutExternalReferencers->Reserve(OutExternalReferencers->Num() + 1);
-					new(*OutExternalReferencers) FReferencerInformation(Object, Count, Referencers);
+					OutExternalReferencers->Emplace(Object, Count, Referencers);
 				}
 			}
 		}
@@ -3902,13 +3902,13 @@ bool StaticExec( UWorld* InWorld, const TCHAR* Cmd, FOutputDevice& Ar )
 				{
 					if (AsteriskPos != INDEX_NONE && (QuestionPos == INDEX_NONE || QuestionPos > AsteriskPos))
 					{
-						new(WildcardPieces) FListPropsWildcardPiece(PropWildcard.Left(AsteriskPos), true);
+						WildcardPieces.Emplace(PropWildcard.Left(AsteriskPos), true);
 						PropWildcard.RightInline(PropWildcard.Len() - AsteriskPos - 1, false);
 						bFound = true;
 					}
 					else if (QuestionPos != INDEX_NONE)
 					{
-						new(WildcardPieces) FListPropsWildcardPiece(PropWildcard.Left(QuestionPos), false);
+						WildcardPieces.Emplace(PropWildcard.Left(QuestionPos), false);
 						PropWildcard.RightInline(PropWildcard.Len() - QuestionPos - 1, false);
 						bFound = true;
 					}
@@ -3917,7 +3917,7 @@ bool StaticExec( UWorld* InWorld, const TCHAR* Cmd, FOutputDevice& Ar )
 			bool bEndedInConstant = (PropWildcard.Len() > 0);
 			if (bEndedInConstant)
 			{
-				new(WildcardPieces) FListPropsWildcardPiece(PropWildcard, false);
+				WildcardPieces.Emplace(PropWildcard, false);
 			}
 
 			// search for matches
diff --git a/Engine/Source/Runtime/CoreUObject/Private/UObject/UObjectGlobals.cpp b/Engine/Source/Runtime/CoreUObject/Private/UObject/UObjectGlobals.cpp
--- a/Engine/Source/Runtime/CoreUObject/Private/UObject/UObjectGlobals.cpp
+++ b/Engine/Source/Runtime/CoreUObject/Private/UObject/UObjectGlobals.cpp
@@ -5140,7 +5140,7 @@ private:
 				{
 					if (!CurrentReferenceInfo)
 					{
-						CurrentReferenceInfo = new(FoundReferencesList->ExternalReferences) FReferencerInformation(CurrentObject);
+						CurrentReferenceInfo = &FoundReferencesList->ExternalReferences.Emplace_GetRef(CurrentObject);
 					}
 					if (InReferencingProperty)
 					{
@@ -5225,7 +5225,7 @@ bool IsReferenced(UObject*& Obj, EObjectFlags KeepFlags, EInternalObjectFlags In
 			else if (OldRef->Referencer->IsIn(Obj))
 			{
 				bReferencedByOuters = true;
-				FReferencerInformation *NewRef = new(FoundReferences->InternalReferences) FReferencerInformation(OldRef->Referencer, OldRef->TotalReferences, OldRef->ReferencingProperties);
+				FReferencerInformation& NewRef = FoundReferences->InternalReferences.Emplace_GetRef(OldRef->Referencer, OldRef->TotalReferences, OldRef->ReferencingProperties);
 				FoundReferences->ExternalReferences.RemoveAt(i);
 				i--;
 			}
diff --git a/Engine/Source/Runtime/MeshUtilitiesCommon/Private/OverlappingCorners.cpp b/Engine/Source/Runtime/MeshUtilitiesCommon/Private/OverlappingCorners.cpp
--- a/Engine/Source/Runtime/MeshUtilitiesCommon/Private/OverlappingCorners.cpp
+++ b/Engine/Source/Runtime/MeshUtilitiesCommon/Private/OverlappingCorners.cpp
@@ -40,7 +40,7 @@ FOverlappingCorners::FOverlappingCorners(const FLayoutUV::IMeshView& MeshView, f
 	VertIndexAndZ.Reserve(NumWedges);
 	for (int32 WedgeIndex = 0; WedgeIndex < NumWedges; WedgeIndex++)
 	{
-		new(VertIndexAndZ)FIndexAndZ(WedgeIndex, MeshView.GetPosition(WedgeIndex));
+		VertIndexAndZ.Emplace(WedgeIndex, MeshView.GetPosition(WedgeIndex));
 	}
 
 	// Sort the vertices by z value
diff --git a/Engine/Source/Runtime/RHI/Private/GPUDefragAllocator.cpp b/Engine/Source/Runtime/RHI/Private/GPUDefragAllocator.cpp
--- a/Engine/Source/Runtime/RHI/Private/GPUDefragAllocator.cpp
+++ b/Engine/Source/Runtime/RHI/Private/GPUDefragAllocator.cpp
@@ -1745,7 +1745,7 @@ void FGPUDefragAllocator::GetMemoryLayout(TArray<FMemoryLayoutElement>& MemoryLa
 	while (Chunk)
 	{
 		EMemoryElementType ChunkType = GetChunkType(Chunk);
-		new (MemoryLayout)FMemoryLayoutElement(Chunk->Size, ChunkType);
+		MemoryLayout.Emplace(Chunk->Size, ChunkType);
 		Chunk = Chunk->NextChunk;
 	}
 }
diff --git a/Engine/Source/Runtime/RenderCore/Private/ShaderParameterMetadataBuilder.cpp b/Engine/Source/Runtime/RenderCore/Private/ShaderParameterMetadataBuilder.cpp
--- a/Engine/Source/Runtime/RenderCore/Private/ShaderParameterMetadataBuilder.cpp
+++ b/Engine/Source/Runtime/RenderCore/Private/ShaderParameterMetadataBuilder.cpp
@@ -10,7 +10,7 @@ void FShaderParametersMetadataBuilder::AddReferencedStruct(
 {
 	NextMemberOffset = Align(NextMemberOffset, SHADER_PARAMETER_STRUCT_ALIGNMENT);
 
-	new(Members) FShaderParametersMetadata::FMember(
+	Members.Emplace(
 		Name,
 		StructMetadata->GetStructTypeName(),
 		__LINE__,
@@ -33,7 +33,7 @@ void FShaderParametersMetadataBuilder::AddIncludedStruct(
 {
 	NextMemberOffset = Align(NextMemberOffset, SHADER_PARAMETER_STRUCT_ALIGNMENT);
 
-	new(Members) FShaderParametersMetadata::FMember(
+	Members.Emplace(
 		TEXT(""),
 		StructMetadata->GetStructTypeName(),
 		__LINE__,
@@ -58,7 +58,7 @@ uint32 FShaderParametersMetadataBuilder::AddNestedStruct(
 	NextMemberOffset = Align(NextMemberOffset, SHADER_PARAMETER_STRUCT_ALIGNMENT);
 	const uint32 ThisMemberOffset = NextMemberOffset;
 
-	new(Members) FShaderParametersMetadata::FMember(
+	Members.Emplace(
 		Name,
 		TEXT(""),
 		__LINE__,
@@ -83,7 +83,7 @@ void FShaderParametersMetadataBuilder::AddBufferSRV(
 {
 	NextMemberOffset = Align(NextMemberOffset, SHADER_PARAMETER_POINTER_ALIGNMENT);
 
-	new(Members) FShaderParametersMetadata::FMember(
+	Members.Emplace(
 		Name,
 		ShaderType,
 		__LINE__,
@@ -107,7 +107,7 @@ void FShaderParametersMetadataBuilder::AddBufferUAV(
 {
 	NextMemberOffset = Align(NextMemberOffset, SHADER_PARAMETER_POINTER_ALIGNMENT);
 
-	new(Members) FShaderParametersMetadata::FMember(
+	Members.Emplace(
 		Name,
 		ShaderType,
 		__LINE__,
@@ -131,7 +131,7 @@ void FShaderParametersMetadataBuilder::AddRDGBufferSRV(
 {
 	NextMemberOffset = Align(NextMemberOffset, SHADER_PARAMETER_POINTER_ALIGNMENT);
 
-	new(Members) FShaderParametersMetadata::FMember(
+	Members.Emplace(
 		Name,
 		ShaderType,
 		__LINE__,
@@ -155,7 +155,7 @@ void FShaderParametersMetadataBuilder::AddRDGBufferUAV(
 {
 	NextMemberOffset = Align(NextMemberOffset, SHADER_PARAMETER_POINTER_ALIGNMENT);
 
-	new(Members) FShaderParametersMetadata::FMember(
+	Members.Emplace(
 		Name,
 		ShaderType,
 		__LINE__,
diff --git a/Engine/Source/Runtime/SandboxFile/Private/IPlatformFileSandboxWrapper.cpp b/Engine/Source/Runtime/SandboxFile/Private/IPlatformFileSandboxWrapper.cpp
--- a/Engine/Source/Runtime/SandboxFile/Private/IPlatformFileSandboxWrapper.cpp
+++ b/Engine/Source/Runtime/SandboxFile/Private/IPlatformFileSandboxWrapper.cpp
@@ -394,7 +394,7 @@ void FSandboxPlatformFile::FindFiles( TArray<FString>& Result, const TCHAR* InFi
 			if( ( bIsDirectory && bDirectories ) ||
 				( !bIsDirectory && bFiles && FString( FilenameOrDirectory ).MatchesWildcard( WildCard ) ) )
 			{
-				new( Result ) FString( FPaths::GetCleanFilename( FilenameOrDirectory ) );
+				Result.Add( FPaths::GetCleanFilename( FilenameOrDirectory ) );
 			}
 			return true;
 		}
diff --git a/Engine/Source/Runtime/StaticMeshDescription/Private/StaticMeshOperations.cpp b/Engine/Source/Runtime/StaticMeshDescription/Private/StaticMeshOperations.cpp
--- a/Engine/Source/Runtime/StaticMeshDescription/Private/StaticMeshOperations.cpp
+++ b/Engine/Source/Runtime/StaticMeshDescription/Private/StaticMeshOperations.cpp
@@ -739,7 +739,7 @@ void FillMeshDescriptionVertexPositionNoDuplicate(const TArray<FVector3f>& RawMe
 
 	for (int32 VertexIndex = 0; VertexIndex < NumVertex; ++VertexIndex)
 	{
-		new(VertIndexAndZ)MeshDescriptionOperationNamespace::FIndexAndZ(VertexIndex, RawMeshVertexPositions[VertexIndex]);
+		VertIndexAndZ.Emplace(VertexIndex, RawMeshVertexPositions[VertexIndex]);
 	}
 
 	// Sort the vertices by z value
@@ -1818,7 +1818,7 @@ void FStaticMeshOperations::FindOverlappingCorners(FOverlappingCorners& OutOverl
 			for (int32 Corner = 0; Corner < 3; ++Corner)
 			{
 				const FVertexInstanceID VertexInstanceID = MeshDescription.GetTriangleVertexInstance(TriangleID, Corner);
-				new(VertIndexAndZ)MeshDescriptionOperationNamespace::FIndexAndZ(WedgeIndex, VertexPositions[MeshDescription.GetVertexInstanceVertex(VertexInstanceID)]);
+				VertIndexAndZ.Emplace(WedgeIndex, VertexPositions[MeshDescription.GetVertexInstanceVertex(VertexInstanceID)]);
 				++WedgeIndex;
 			}
 		}
@@ -2465,7 +2465,7 @@ void FStaticMeshOperations::BuildWeldedVertexIDRemap(const FMeshDescription& Mes
 
 	for (const FVertexID VertexID : MeshDescription.Vertices().GetElementIDs())
 	{
-		new(VertIndexAndZ)MeshDescriptionOperationNamespace::FIndexAndZ(VertexID.GetValue(), VertexPositions[VertexID]);
+		VertIndexAndZ.Emplace(VertexID.GetValue(), VertexPositions[VertexID]);
 	}
 
 	// Sort the vertices by z value
