Typecast warning fixups

diff --git a/Engine/Source/Developer/TextureBuild/Private/TextureBuildFunction.cpp b/Engine/Source/Developer/TextureBuild/Private/TextureBuildFunction.cpp
--- a/Engine/Source/Developer/TextureBuild/Private/TextureBuildFunction.cpp
+++ b/Engine/Source/Developer/TextureBuild/Private/TextureBuildFunction.cpp
@@ -234,7 +234,7 @@ static bool TryReadTextureSourceFromCompactBinary(FCbFieldView Source, UE::Deriv
 	}
 
 	FCbArrayView MipsCbArrayView = Source["Mips"].AsArrayView();
-	OutMips.Reserve(MipsCbArrayView.Num());
+	OutMips.Reserve(IntCastChecked<int32>(MipsCbArrayView.Num()));
 	for (FCbFieldView MipsCbArrayIt : MipsCbArrayView)
 	{
 		FCbObjectView MipCbObjectView = MipsCbArrayIt.AsObjectView();
diff --git a/Engine/Source/Developer/TextureBuild/TextureBuild.Build.cs b/Engine/Source/Developer/TextureBuild/TextureBuild.Build.cs
--- a/Engine/Source/Developer/TextureBuild/TextureBuild.Build.cs
+++ b/Engine/Source/Developer/TextureBuild/TextureBuild.Build.cs
@@ -6,6 +6,8 @@ public class TextureBuild : ModuleRules
 {
 	public TextureBuild(ReadOnlyTargetRules Target) : base(Target)
 	{
+		UnsafeTypeCastWarningLevel = WarningLevel.Warning;
+
 		PrivateIncludePathModuleNames.AddRange(new string[]
 		{
 			"Engine",
diff --git a/Engine/Source/Developer/TextureCompressor/Private/TextureCompressorModule.cpp b/Engine/Source/Developer/TextureCompressor/Private/TextureCompressorModule.cpp
--- a/Engine/Source/Developer/TextureCompressor/Private/TextureCompressorModule.cpp
+++ b/Engine/Source/Developer/TextureCompressor/Private/TextureCompressorModule.cpp
@@ -689,7 +689,7 @@ static void GenerateMip2x2(const FImageView2D& SourceImageData, FImageView2D& De
 }
 
 static void AllocateTempForMips(
-	TArray<FLinearColor>& TempData,
+	TArray64<FLinearColor>& TempData,
 	int32 SourceSizeX,
 	int32 SourceSizeY,
 	int32 DestSizeX,
@@ -838,7 +838,7 @@ static void GenerateMipSharpened(
 template <EMipGenAddressMode AddressMode, int KernelSize = 0>
 static void GenerateMipSharpenedSeparable(
 	const FImageView2D& SourceImageData,
-	TArray<FLinearColor>& TempData,
+	TArray64<FLinearColor>& TempData,
 	FImageView2D& DestImageData,
 	const FImageKernel2D& Kernel,
 	FVector4f AlphaScale,
@@ -921,7 +921,7 @@ static void GenerateMipSharpenedSeparable(
 template <EMipGenAddressMode AddressMode>
 static void GenerateSharpenedMipB8G8R8A8Templ(
 	const FImageView2D& SourceImageData, 
-	TArray<FLinearColor>& TempData,
+	TArray64<FLinearColor>& TempData,
 	FImageView2D& DestImageData, 
 	bool bDoScaleMipsForAlphaCoverage,
 	const FVector4f AlphaCoverages,
@@ -1024,7 +1024,7 @@ static void GenerateMipBorder(
 static void GenerateSharpenedMipB8G8R8A8(
 	const FImageView2D& SourceImageData, 
 	const FImageView2D& SourceImageData2, // Only used with volume texture.
-	TArray<FLinearColor>& TempData,
+	TArray64<FLinearColor>& TempData,
 	FImageView2D& DestImageData, 
 	EMipGenAddressMode AddressMode, 
 	bool bDoScaleMipsForAlphaCoverage,
@@ -1260,7 +1260,7 @@ static void GenerateTopMip(const FImage& SrcImage, FImage& DestImage, const FTex
 	const bool bUnfiltered = Settings.MipGenSettings == TMGS_Unfiltered;
 	const bool bUseNewMipFilter = Settings.bUseNewMipFilter;
 
-	TArray<FLinearColor> TempData;
+	TArray64<FLinearColor> TempData;
 	AllocateTempForMips(TempData, SrcImage.SizeX, SrcImage.SizeY, DestImage.SizeX, DestImage.SizeY, false, KernelDownsample, 1, bSharpenWithoutColorShift, bUnfiltered, bUseNewMipFilter);
 
 	for (int32 SliceIndex = 0; SliceIndex < SrcImage.NumSlices; ++SliceIndex)
@@ -1483,7 +1483,7 @@ static void DownscaleImage(const FImage& SrcImage, FImage& DstImage, const FText
 	FImageKernel2D AvgKernel;
 	AvgKernel.BuildSeparatableGaussWithSharpen(2);
 
-	TArray<FLinearColor> TempData;
+	TArray64<FLinearColor> TempData;
 	AllocateTempForMips(TempData, SrcImage.SizeX, SrcImage.SizeY, FMath::Max(1, SrcImage.SizeX / 2), FMath::Max(1, SrcImage.SizeY / 2), false, AvgKernel, 2, false, bUnfiltered, bUseNewMipFilter);
 
 	int32 NumIterations = 0;
@@ -1712,8 +1712,8 @@ void ITextureCompressorModule::GenerateMipChain(
 	int32 FirstMipSizeX = FMath::Max(BaseMip.SizeX>>1,1);
 	int32 FirstMipSizeY = FMath::Max(BaseMip.SizeY>>1,1);
 
-	TArray<FLinearColor> DownsampleTempData;
-	TArray<FLinearColor> AverageTempData;
+	TArray64<FLinearColor> DownsampleTempData;
+	TArray64<FLinearColor> AverageTempData;
 	const bool bDoScaleMipsForAlphaCoverage = Settings.bDoScaleMipsForAlphaCoverage;
 	const bool bSharpenWithoutColorShift = Settings.bSharpenWithoutColorShift;
 	const bool bUnfiltered = Settings.MipGenSettings == TMGS_Unfiltered;
@@ -1775,7 +1775,7 @@ void ITextureCompressorModule::GenerateMipChain(
 			FImageView2D DestView(DestImage, SliceIndex);
 
 			// lambda to generate one step :
-			auto MakeDestImageMip = [&](FImageView2D & ArgDestView,TArray<FLinearColor> & ArgTempData,FImageKernel2D & ArgKernel) {
+			auto MakeDestImageMip = [&](FImageView2D & ArgDestView,TArray64<FLinearColor> & ArgTempData,FImageKernel2D & ArgKernel) {
 				// DestView is the output mip
 				GenerateSharpenedMipB8G8R8A8(
 					IntermediateSrcView, 
@@ -1973,8 +1973,8 @@ struct FImageViewLongLat
 		// acos returns in [0,PI]
 
 		const float invPI = 1.f/PI;
-		float X = (1.f + atan2f(NormalizedDirection.X, - NormalizedDirection.Z) * invPI) * 0.5f * SizeX;
-		float Y = acosf(NormalizedDirection.Y)*invPI * SizeY;
+		float X = (1.f + atan2f(static_cast<float>(NormalizedDirection.X), static_cast<float>(-NormalizedDirection.Z)) * invPI) * 0.5f * SizeX;
+		float Y = acosf(static_cast<float>(NormalizedDirection.Y))*invPI * SizeY;
 
 		return LookupFiltered(X, Y);
 	}
@@ -1983,8 +1983,8 @@ struct FImageViewLongLat
 	{
 		// this does the math in doubles then stores to floats :
 		//	that was probably a mistake, but leave it to avoid patches
-		float X = (1 + atan2(NormalizedDirection.X, - NormalizedDirection.Z) / PI) / 2 * SizeX;
-		float Y = acos(NormalizedDirection.Y) / PI * SizeY;
+		float X = (1 + atan2(static_cast<float>(NormalizedDirection.X), static_cast<float>(-NormalizedDirection.Z)) / PI) / 2 * SizeX;
+		float Y = acos(static_cast<float>(NormalizedDirection.Y)) / PI * SizeY;
 
 		return LookupFiltered(X, Y);
 	}
@@ -1993,7 +1993,7 @@ struct FImageViewLongLat
 // transform world space vector to a space relative to the face
 static inline FVector TransformSideToWorldSpace(uint32 CubemapFace, const FVector & InDirection)
 {
-	float x = InDirection.X, y = InDirection.Y, z = InDirection.Z;
+	double x = InDirection.X, y = InDirection.Y, z = InDirection.Z;
 
 	FVector Ret;
 
@@ -2017,7 +2017,7 @@ static inline FVector TransformSideToWorldSpace(uint32 CubemapFace, const FVecto
 static inline FVector TransformWorldToSideSpace(uint32 CubemapFace, const FVector & InDirection)
 {
 	// undo Unreal way (z and y are flipped)
-	float x = InDirection.X, y = InDirection.Z, z = InDirection.Y;
+	double x = InDirection.X, y = InDirection.Z, z = InDirection.Y;
 
 	FVector Ret;
 
@@ -2138,7 +2138,7 @@ public:
 		
 		FVector DirectionSS = ComputeSSCubeDirectionAtTexelCenter(x, y, InvFullExtent);
 
-		float DotValue = ConeAxisSS | DirectionSS;
+		float DotValue = static_cast<float>(ConeAxisSS | DirectionSS);
 
 		if(DotValue > DirDot)
 		{
@@ -2243,7 +2243,7 @@ static FLinearColor IntegrateAngularArea(FImage& Image, FVector FilterDirectionW
 // @return 2 * computed triangle area 
 static inline float TriangleArea2_3D(FVector A, FVector B, FVector C)
 {
-	return ((A-B) ^ (C-B)).Size();
+	return static_cast<float>(((A-B) ^ (C-B)).Size());
 }
 
 static inline float ComputeTexelArea(uint32 x, uint32 y, float InvSideExtentMul2)
@@ -3087,7 +3087,7 @@ static bool ApplyCompositeTexture(FImage& DestRoughness, const FImage& SourceNor
 			return false;
 		}
 
-		TArray<float> SourceNormalLengths;
+		TArray64<float> SourceNormalLengths;
 
 		{
 			int64 SourceNormalCount = (int64) SourceNormals.SizeX * SourceNormals.SizeY;
@@ -3144,7 +3144,7 @@ void FTextureBuildSettings::GetEncodedTextureDescriptionWithPixelFormat(FEncoded
 	TextureDescription.bCubeMap = bCubemap;
 	TextureDescription.bTextureArray = bTextureArray;
 	TextureDescription.bVolumeTexture = bVolume;
-	TextureDescription.NumMips = InMipCount;
+	TextureDescription.NumMips = IntCastChecked<uint8>(InMipCount);
 	TextureDescription.PixelFormat = InEncodedPixelFormat;
 
 	TextureDescription.TopMipSizeX = InEncodedMip0SizeX;
@@ -3360,7 +3360,11 @@ static void NormalizeMip(FImage& InOutMip)
 			// GetSafeNormal returns Vec(0,0,0) by default for tiny input, instead return flat/up
 			Normal = Normal.GetSafeNormal(UE_SMALL_NUMBER,NormalIfZero);
 
-			Color = FLinearColor(Normal.X * 0.5f + 0.5f, Normal.Y * 0.5f + 0.5f, Normal.Z * 0.5f + 0.5f, Color.A);
+			Color = FLinearColor(
+				static_cast<float>(Normal.X * 0.5f + 0.5f), 
+				static_cast<float>(Normal.Y * 0.5f + 0.5f), 
+				static_cast<float>(Normal.Z * 0.5f + 0.5f), 
+				Color.A);
 		}
 	}, EParallelForFlags::Unbalanced);
 }
diff --git a/Engine/Source/Developer/TextureCompressor/TextureCompressor.Build.cs b/Engine/Source/Developer/TextureCompressor/TextureCompressor.Build.cs
--- a/Engine/Source/Developer/TextureCompressor/TextureCompressor.Build.cs
+++ b/Engine/Source/Developer/TextureCompressor/TextureCompressor.Build.cs
@@ -6,6 +6,8 @@ public class TextureCompressor : ModuleRules
 {
 	public TextureCompressor(ReadOnlyTargetRules Target) : base(Target)
 	{
+		UnsafeTypeCastWarningLevel = WarningLevel.Warning;
+
 		PrivateIncludePathModuleNames.AddRange(new string[]
 		{
 			"Engine",
diff --git a/Engine/Source/Developer/TextureFormatDXT/Private/TextureFormatDXT.cpp b/Engine/Source/Developer/TextureFormatDXT/Private/TextureFormatDXT.cpp
--- a/Engine/Source/Developer/TextureFormatDXT/Private/TextureFormatDXT.cpp
+++ b/Engine/Source/Developer/TextureFormatDXT/Private/TextureFormatDXT.cpp
@@ -71,7 +71,7 @@ static FName GSupportedTextureFormatNames[] =
  */
 struct FNVOutputHandler : public nvtt::OutputHandler
 {
-	explicit FNVOutputHandler( uint8* InBuffer, int32 InBufferSize )
+	explicit FNVOutputHandler( uint8* InBuffer, int64 InBufferSize )
 		: Buffer(InBuffer)
 		, BufferEnd(InBuffer + InBufferSize)
 	{
@@ -142,7 +142,7 @@ public:
 		bool bSRGB,
 		bool bIsNormalMap,
 		uint8* OutBuffer,
-		int32 BufferSize, 
+		int64 BufferSize, 
 		bool bPreview = false)
 		: OutputHandler(OutBuffer, BufferSize)
 	{
diff --git a/Engine/Source/Developer/TextureFormatDXT/TextureFormatDXT.Build.cs b/Engine/Source/Developer/TextureFormatDXT/TextureFormatDXT.Build.cs
--- a/Engine/Source/Developer/TextureFormatDXT/TextureFormatDXT.Build.cs
+++ b/Engine/Source/Developer/TextureFormatDXT/TextureFormatDXT.Build.cs
@@ -6,6 +6,8 @@ public class TextureFormatDXT : ModuleRules
 {
 	public TextureFormatDXT(ReadOnlyTargetRules Target) : base(Target)
 	{
+		UnsafeTypeCastWarningLevel = WarningLevel.Warning;
+
 		PrivateIncludePathModuleNames.AddRange(new string[]
 		{
 			"DerivedDataCache",
diff --git a/Engine/Source/Developer/TextureFormatIntelISPCTexComp/Private/TextureFormatIntelISPCTexComp.cpp b/Engine/Source/Developer/TextureFormatIntelISPCTexComp/Private/TextureFormatIntelISPCTexComp.cpp
--- a/Engine/Source/Developer/TextureFormatIntelISPCTexComp/Private/TextureFormatIntelISPCTexComp.cpp
+++ b/Engine/Source/Developer/TextureFormatIntelISPCTexComp/Private/TextureFormatIntelISPCTexComp.cpp
@@ -425,7 +425,7 @@ static uint16 GetDefaultCompressionBySizeValue(FCbObjectView InFormatConfigOverr
 		checkf(FieldView.HasValue(), TEXT("Missing DefaultASTCQualityBySize key from FormatConfigOverride"));
 		int32 CompressionModeValue = FieldView.AsInt32();
 		checkf(!FieldView.HasError(), TEXT("Failed to parse DefaultASTCQualityBySize value from FormatConfigOverride"));
-		return CompressionModeValue;
+		return IntCastChecked<uint16>(CompressionModeValue);
 	}
 	else
 	{
@@ -442,7 +442,7 @@ static uint16 GetDefaultCompressionBySizeValue(FCbObjectView InFormatConfigOverr
 			return FMath::Min<uint32>(CompressionModeValue, MAX_QUALITY_BY_SIZE);
 		};
 
-		static int32 CompressionModeValue = GetCompressionModeValue();
+		static uint16 CompressionModeValue = IntCastChecked<uint16>(GetCompressionModeValue());
 
 		return CompressionModeValue;
 	}
@@ -576,9 +576,9 @@ static void IntelASTCCompressScans(FASTCEncoderSettings* pEncSettings, FImage* p
 					FVector Normal = FVector(pInTexelsSwap[2] / 255.0f * 2.0f - 1.0f, pInTexelsSwap[1] / 255.0f * 2.0f - 1.0f, pInTexelsSwap[0] / 255.0f * 2.0f - 1.0f);
 					Normal = Normal.GetSafeNormal();
 					pInTexelsSwap[0] = 0;
-					pInTexelsSwap[1] = FMath::RoundToInt((Normal.Y * 0.5f + 0.5f) * 255.f);
+					pInTexelsSwap[1] = static_cast<uint8>(FMath::RoundToInt((Normal.Y * 0.5f + 0.5f) * 255.f));
 					pInTexelsSwap[2] = 0;
-					pInTexelsSwap[3] = FMath::RoundToInt((Normal.X * 0.5f + 0.5f) * 255.f);
+					pInTexelsSwap[3] = static_cast<uint8>(FMath::RoundToInt((Normal.X * 0.5f + 0.5f) * 255.f));
 
 					pInTexelsSwap += 4;
 				}
@@ -618,8 +618,8 @@ static void IntelASTCCompressScans(FASTCEncoderSettings* pEncSettings, FImage* p
 				{
 					FVector Normal = FVector(pInTexelsSwap[2] / 255.0f * 2.0f - 1.0f, pInTexelsSwap[1] / 255.0f * 2.0f - 1.0f, pInTexelsSwap[0] / 255.0f * 2.0f - 1.0f);
 					Normal = Normal.GetSafeNormal();
-					pInTexelsSwap[0] = FMath::RoundToInt((Normal.X * 0.5f + 0.5f) * 255.f);
-					pInTexelsSwap[1] = FMath::RoundToInt((Normal.Y * 0.5f + 0.5f) * 255.f);
+					pInTexelsSwap[0] = static_cast<uint8>(FMath::RoundToInt((Normal.X * 0.5f + 0.5f) * 255.f));
+					pInTexelsSwap[1] = static_cast<uint8>(FMath::RoundToInt((Normal.Y * 0.5f + 0.5f) * 255.f));
 					pInTexelsSwap[2] = 0;
 					pInTexelsSwap[3] = 255;
 
diff --git a/Engine/Source/Developer/TextureFormatIntelISPCTexComp/TextureFormatIntelISPCTexComp.Build.cs b/Engine/Source/Developer/TextureFormatIntelISPCTexComp/TextureFormatIntelISPCTexComp.Build.cs
--- a/Engine/Source/Developer/TextureFormatIntelISPCTexComp/TextureFormatIntelISPCTexComp.Build.cs
+++ b/Engine/Source/Developer/TextureFormatIntelISPCTexComp/TextureFormatIntelISPCTexComp.Build.cs
@@ -6,6 +6,8 @@ public class TextureFormatIntelISPCTexComp : ModuleRules
 {
 	public TextureFormatIntelISPCTexComp(ReadOnlyTargetRules Target) : base(Target)
 	{
+		UnsafeTypeCastWarningLevel = WarningLevel.Warning;
+
 		PrivateIncludePathModuleNames.AddRange(new string[]
 		{
 			"DerivedDataCache",
diff --git a/Engine/Source/Editor/TextureEditor/Private/TextureEditorToolkit.cpp b/Engine/Source/Editor/TextureEditor/Private/TextureEditorToolkit.cpp
--- a/Engine/Source/Editor/TextureEditor/Private/TextureEditorToolkit.cpp
+++ b/Engine/Source/Editor/TextureEditor/Private/TextureEditorToolkit.cpp
@@ -343,9 +343,9 @@ void FTextureEditorToolkit::InitTextureEditor( const EToolkitMode::Type Mode, co
 
 void FTextureEditorToolkit::CalculateTextureDimensions(int32& OutWidth, int32& OutHeight, int32& OutDepth, int32& OutArraySize, bool bInIncludeBorderSize) const
 {
-	OutWidth = Texture->GetSurfaceWidth();
-	OutHeight = Texture->GetSurfaceHeight();
-	OutDepth = Texture->GetSurfaceDepth();
+	OutWidth = static_cast<int32>(Texture->GetSurfaceWidth());
+	OutHeight = static_cast<int32>(Texture->GetSurfaceHeight());
+	OutDepth = static_cast<int32>(Texture->GetSurfaceDepth());
 	OutArraySize = IsArrayTexture() ? (IsCubeTexture() ? Texture->GetSurfaceArraySize() / 6 : Texture->GetSurfaceArraySize()) : 0;
 	const int32 BorderSize = GetDefault<UTextureEditorSettings>()->TextureBorderEnabled ? 1 : 0;
 
@@ -390,8 +390,8 @@ void FTextureEditorToolkit::CalculateTextureDimensions(int32& OutWidth, int32& O
 	}
 	else
 	{
-		OutWidth = PreviewEffectiveTextureWidth * Zoom;
-		OutHeight = PreviewEffectiveTextureHeight * Zoom;
+		OutWidth = static_cast<int32>(PreviewEffectiveTextureWidth * Zoom);
+		OutHeight = static_cast<int32>(PreviewEffectiveTextureHeight * Zoom);
 	}
 
 	if (bInIncludeBorderSize)
@@ -724,9 +724,9 @@ void FTextureEditorToolkit::PopulateQuickInfo( )
 	const bool bIsArray = IsArrayTexture();
 	const bool bIsCube = IsCubeTexture();
 
-	const int32 SurfaceWidth = Texture->GetSurfaceWidth();
-	const int32 SurfaceHeight = Texture->GetSurfaceHeight();
-	const int32 SurfaceDepth = Texture->GetSurfaceDepth();
+	const int32 SurfaceWidth = static_cast<int32>(Texture->GetSurfaceWidth());
+	const int32 SurfaceHeight = static_cast<int32>(Texture->GetSurfaceHeight());
+	const int32 SurfaceDepth = static_cast<int32>(Texture->GetSurfaceDepth());
 	const int32 NumSurfaces = Texture->GetSurfaceArraySize();
 	const int32 ArraySize = bIsArray ? (bIsCube ? NumSurfaces / 6 : NumSurfaces) : 0;
 
@@ -1208,8 +1208,17 @@ void FTextureEditorToolkit::CreateInternalWidgets()
 
 	OodleCompressionLevel = FOodleDataCompression::ECompressionLevel::Optimal3;
 	const TCHAR* LevelName;
-	{		 
-		FOodleDataCompression::ECompressionLevelFromValue(ProjectSettings->PackageCompressionLevel_Distribution, OodleCompressionLevel);
+	{
+		int8 PackageCompressionLevel_Distribution8 = 0;
+		if (IntFitsIn<int8>(ProjectSettings->PackageCompressionLevel_Distribution))
+		{
+			PackageCompressionLevel_Distribution8 = static_cast<int8>(ProjectSettings->PackageCompressionLevel_Distribution);
+		}
+		else
+		{
+			UE_LOG(LogTextureEditor, Warning, TEXT("PackageCompressionLevel_Distribution in Project Settings is an invalid value = %d, must be [-128, 127]"), ProjectSettings->PackageCompressionLevel_Distribution);
+		}
+		FOodleDataCompression::ECompressionLevelFromValue(PackageCompressionLevel_Distribution8, OodleCompressionLevel);
 		FOodleDataCompression::ECompressionLevelToString(OodleCompressionLevel, &LevelName);
 	}
 
@@ -2659,7 +2668,7 @@ void FTextureEditorToolkit::HandleZoomSliderChanged(float SliderValue)
 
 float FTextureEditorToolkit::HandleZoomSliderValue() const
 {
-	float ZoomValue = CalculateDisplayedZoomLevel();
+	double ZoomValue = CalculateDisplayedZoomLevel();
 	double Octaves = log2( MaxZoom/MinZoom );
 	double SliderValue = log2( ZoomValue/MinZoom ) / Octaves;
 
@@ -2675,7 +2684,7 @@ void FTextureEditorToolkit::EditorOodleSettingsEffortChanged(int32 NewValue, ESe
 {
 	bool bChanged = CustomEncoding->OodleEncodeEffort != NewValue;
 
-	CustomEncoding->OodleEncodeEffort = NewValue;
+	CustomEncoding->OodleEncodeEffort = IntCastChecked<uint8>(NewValue);
 
 	if (CustomEncoding->bUseCustomEncode || bChanged)
 	{
@@ -2691,7 +2700,7 @@ int32 FTextureEditorToolkit::GetEditorOodleSettingsTiling() const
 void FTextureEditorToolkit::EditorOodleSettingsTilingChanged(int32 NewValue, ESelectInfo::Type SelectionType)
 {
 	bool bChanged = CustomEncoding->OodleUniversalTiling != NewValue;
-	CustomEncoding->OodleUniversalTiling = NewValue;
+	CustomEncoding->OodleUniversalTiling = IntCastChecked<uint8>(NewValue);
 
 	if (CustomEncoding->bUseCustomEncode && bChanged)
 	{
@@ -3203,7 +3212,7 @@ void FTextureEditorToolkit::PackagingSettingsChanged(TSharedPtr<FString> Selecti
 	if (Selection.IsValid())
 	{
 		UProjectPackagingSettings const* ProjectSettings = GetDefault<UProjectPackagingSettings>();
-		int8 CompressionLevelFromSettings = (int8)FOodleDataCompression::ECompressionLevel::Optimal3;
+		int32 CompressionLevelFromSettings = (int32)FOodleDataCompression::ECompressionLevel::Optimal3;
 		if (*Selection == TEXT("DebugDevelopment"))
 		{
 			CompressionLevelFromSettings = ProjectSettings->PackageCompressionLevel_DebugDevelopment;
@@ -3217,8 +3226,14 @@ void FTextureEditorToolkit::PackagingSettingsChanged(TSharedPtr<FString> Selecti
 			CompressionLevelFromSettings = ProjectSettings->PackageCompressionLevel_Distribution;
 		}
 
+		if (IntFitsIn<int8>(CompressionLevelFromSettings) == false)
+		{
+			UE_LOG(LogTextureEditor, Warning, TEXT("PackageCompressionLevel_%s in Project Settings is an invalid value = %d, must be [-128, 127]"), **Selection, CompressionLevelFromSettings);
+			CompressionLevelFromSettings = (int32)FOodleDataCompression::ECompressionLevel::Optimal3;
+		}
+
 		FOodleDataCompression::ECompressionLevel OldLevel = OodleCompressionLevel;
-		FOodleDataCompression::ECompressionLevelFromValue(CompressionLevelFromSettings, OodleCompressionLevel);
+		FOodleDataCompression::ECompressionLevelFromValue(static_cast<int8>(CompressionLevelFromSettings), OodleCompressionLevel);
 
 		const TCHAR* LevelName;
 		FOodleDataCompression::ECompressionLevelToString(OodleCompressionLevel, &LevelName);
diff --git a/Engine/Source/Editor/TextureEditor/TextureEditor.Build.cs b/Engine/Source/Editor/TextureEditor/TextureEditor.Build.cs
--- a/Engine/Source/Editor/TextureEditor/TextureEditor.Build.cs
+++ b/Engine/Source/Editor/TextureEditor/TextureEditor.Build.cs
@@ -6,6 +6,8 @@ public class TextureEditor : ModuleRules
 {
 	public TextureEditor(ReadOnlyTargetRules Target) : base(Target)
 	{
+		UnsafeTypeCastWarningLevel = WarningLevel.Warning;
+
 		DynamicallyLoadedModuleNames.AddRange(
 			new string[] {
 				"MainFrame",
