Clang specific cast warnings, and fixing an incorrect cast fix from last checkin.

diff --git a/Engine/Source/Developer/TextureCompressor/Private/TextureCompressorModule.cpp b/Engine/Source/Developer/TextureCompressor/Private/TextureCompressorModule.cpp
--- a/Engine/Source/Developer/TextureCompressor/Private/TextureCompressorModule.cpp
+++ b/Engine/Source/Developer/TextureCompressor/Private/TextureCompressorModule.cpp
@@ -269,11 +269,11 @@ private:
 	// support even and non even sized filters
 	static void BuildGaussian1D(float *InOutTable, uint32 TableSize, float Sum, float Variance)
 	{
-		float Center = TableSize * 0.5f - 0.5f;
+		float Center = (float)TableSize * 0.5f - 0.5f;
 		float CurrentSum = 0;
 		for(uint32 i = 0; i < TableSize; ++i)
 		{
-			float Actual = NormalDistribution(i - Center, Variance);
+			float Actual = NormalDistribution((float)i - Center, Variance);
 			InOutTable[i] = Actual;
 			CurrentSum += Actual;
 		}
@@ -1323,12 +1323,12 @@ static FLinearColor Bilerp(
 //  range of X is [-0.5,-0.5] to [Width-0.5,Height-0.5]
 static FLinearColor LookupSourceMipBilinear(const FImageView2D& SourceImageData, float X, float Y)
 {
-	X = FMath::Clamp(X, 0.f, SourceImageData.SizeX - 1.f);
-	Y = FMath::Clamp(Y, 0.f, SourceImageData.SizeY - 1.f);
+	X = FMath::Clamp(X, 0.f, (float)SourceImageData.SizeX - 1.f);
+	Y = FMath::Clamp(Y, 0.f, (float)SourceImageData.SizeY - 1.f);
 	int32 IntX0 = FMath::TruncToInt32(X);
 	int32 IntY0 = FMath::TruncToInt32(Y);
-	float FractX = X - IntX0;
-	float FractY = Y - IntY0;
+	float FractX = X - (float)IntX0;
+	float FractY = Y - (float)IntY0;
 	int32 IntX1 = FMath::Min(IntX0+1, SourceImageData.SizeX-1);
 	int32 IntY1 = FMath::Min(IntY0+1, SourceImageData.SizeY-1);
 	
@@ -1343,8 +1343,8 @@ static FLinearColor LookupSourceMipBilinear(const FImageView2D& SourceImageData,
 // UV range is [0,1] , first pixel center is at 0.5/W
 static FLinearColor LookupSourceMipBilinearUV(const FImageView2D& SourceImageData, float U, float V)
 {
-	float X = U * SourceImageData.SizeX - 0.5f;
-	float Y = V * SourceImageData.SizeY - 0.5f;
+	float X = U * (float)SourceImageData.SizeX - 0.5f;
+	float Y = V * (float)SourceImageData.SizeY - 0.5f;
 	return LookupSourceMipBilinear(SourceImageData,X,Y);
 }
 
@@ -1352,12 +1352,12 @@ static FLinearColor LookupSourceMipBilinearUV(const FImageView2D& SourceImageDat
 //  range of X is [-0.5,-0.5] to [Width-0.5,Height-0.5]
 static float LookupFloatBilinear(const float * FloatPlane,int32 SizeX,int32 SizeY, float X, float Y)
 {
-	X = FMath::Clamp(X, 0.f, SizeX - 1.f);
-	Y = FMath::Clamp(Y, 0.f, SizeY - 1.f);
+	X = FMath::Clamp(X, 0.f, (float)SizeX - 1.f);
+	Y = FMath::Clamp(Y, 0.f, (float)SizeY - 1.f);
 	int32 IntX0 = FMath::TruncToInt32(X);
 	int32 IntY0 = FMath::TruncToInt32(Y);
-	float FractX = X - IntX0;
-	float FractY = Y - IntY0;
+	float FractX = X - (float)IntX0;
+	float FractY = Y - (float)IntY0;
 	int32 IntX1 = FMath::Min(IntX0+1, SizeX-1);
 	int32 IntY1 = FMath::Min(IntY0+1, SizeY-1);
 	
@@ -1374,8 +1374,8 @@ static float LookupFloatBilinear(const float * FloatPlane,int32 SizeX,int32 Size
 // UV range is [0,1] , first pixel center is at 0.5/W
 static float LookupFloatBilinearUV(const float * FloatPlane,int32 SizeX,int32 SizeY, float U, float V)
 {
-	float X = U * SizeX - 0.5f;
-	float Y = V * SizeY - 0.5f;
+	float X = U * (float)SizeX - 0.5f;
+	float Y = V * (float)SizeY - 0.5f;
 	return LookupFloatBilinear(FloatPlane,SizeX,SizeY,X,Y);
 }
 
@@ -1402,8 +1402,8 @@ static float GetDownscaleFinalSizeAndClampedDownscale(int32 SrcImageWidth, int32
 	float Downscale = FMath::Clamp(Settings.Downscale, 1.f, 8.f);
 
 	// note: more accurate would be to use FMath::Max(1, FMath::RoundToInt(SrcImage.SizeX / Downscale))
-	int32 FinalSizeX = FMath::CeilToInt(SrcImageWidth / Downscale);
-	int32 FinalSizeY = FMath::CeilToInt(SrcImageHeight / Downscale);
+	int32 FinalSizeX = FMath::CeilToInt((float)SrcImageWidth / Downscale);
+	int32 FinalSizeY = FMath::CeilToInt((float)SrcImageHeight / Downscale);
 
 	// compute final size respecting image block size
 	if (Settings.BlockSize > 1
@@ -1469,7 +1469,7 @@ static void DownscaleImage(const FImage& SrcImage, FImage& DstImage, const FText
 
 	// recompute Downscale factor because it may have changed due to block alignment
 	// note: if aspect ratio was not exactly preserved, this could differ in X and Y
-	Downscale = (float)SrcImage.SizeX / FinalSizeX;
+	Downscale = (float)SrcImage.SizeX / (float)FinalSizeX;
 	
 	// while desired final size is > 2X smaller, do 2X resizes using the standard mip gen filter code
 
@@ -1530,7 +1530,7 @@ static void DownscaleImage(const FImage& SrcImage, FImage& DstImage, const FText
 	if (Settings.DownscaleOptions >= (uint8)ETextureDownscaleOptions::Sharpen0 && Settings.DownscaleOptions <= (uint8)ETextureDownscaleOptions::Sharpen10)
 	{
 		// 0 .. 2.0f
-		Sharpening = ((int32)Settings.DownscaleOptions - (int32)ETextureDownscaleOptions::Sharpen0) * 0.2f;
+		Sharpening = (float)((int32)Settings.DownscaleOptions - (int32)ETextureDownscaleOptions::Sharpen0) * 0.2f;
 		KernelSize = 8;
 	}
 	
@@ -1551,7 +1551,7 @@ static void DownscaleImage(const FImage& SrcImage, FImage& DstImage, const FText
 
 	// recompute Downscale factor again for final arbitrary-factor resizes :
 	// note: if aspect ratio was not exactly preserved, this could differ in X and Y
-	Downscale = (float)ImageChain[0]->SizeX / FinalSizeX;
+	Downscale = (float)ImageChain[0]->SizeX / (float)FinalSizeX;
 
 	FImageView2D SrcImageData(*ImageChain[0], 0);
 	FImageView2D DstImageData(*ImageChain[1], 0);
@@ -1559,12 +1559,12 @@ static void DownscaleImage(const FImage& SrcImage, FImage& DstImage, const FText
 	// @todo OodleImageResize : not sure this is a correct image resize without shift; does it get pixel center offsets right?
 	for (int32 Y = 0; Y < FinalSizeY; ++Y)
 	{
-		float SourceY = Y * Downscale;
+		float SourceY = (float)Y * Downscale;
 		int32 IntSourceY = FMath::RoundToInt(SourceY);
 		
 		for (int32 X = 0; X < FinalSizeX; ++X)
 		{
-			float SourceX = X * Downscale;
+			float SourceX = (float)X * Downscale;
 			int32 IntSourceX = FMath::RoundToInt(SourceX);
 
 			FLinearColor FilteredColor(0,0,0,0);
@@ -1584,7 +1584,7 @@ static void DownscaleImage(const FImage& SrcImage, FImage& DstImage, const FText
 					for (uint32 KernelX = 0; KernelX < KernelSharpen.GetFilterTableSize();  ++KernelX)
 					{
 						float Weight = KernelSharpen.GetAt(KernelX, KernelY);
-						FLinearColor Sample = LookupSourceMipBilinear(SrcImageData, SourceX + KernelX - KernelCenter, SourceY + KernelY - KernelCenter);
+						FLinearColor Sample = LookupSourceMipBilinear(SrcImageData, SourceX + (float)KernelX - (float)KernelCenter, SourceY + (float)KernelY - (float)KernelCenter);
 						FilteredColor += Weight	* Sample;
 					}
 				}
@@ -1931,8 +1931,8 @@ struct FImageViewLongLat
 		int32 X0 = (int32)X;
 		int32 Y0 = (int32)Y;
 
-		float FracX = X - X0;
-		float FracY = Y - Y0;
+		float FracX = X - (float)X0;
+		float FracY = Y - (float)Y0;
 
 		int32 X1 = X0 + 1;
 		int32 Y1 = Y0 + 1;
@@ -1971,10 +1971,11 @@ struct FImageViewLongLat
 	{
 		// atan2 returns in [-PI,PI]
 		// acos returns in [0,PI]
+		const FVector3f NormalizedDirectionFloat { NormalizedDirection };
 
 		const float invPI = 1.f/PI;
-		float X = (1.f + atan2f(static_cast<float>(NormalizedDirection.X), static_cast<float>(-NormalizedDirection.Z)) * invPI) * 0.5f * SizeX;
-		float Y = acosf(static_cast<float>(NormalizedDirection.Y))*invPI * SizeY;
+		float X = (1.f + atan2f(NormalizedDirectionFloat.X, -NormalizedDirectionFloat.Z) * invPI) * 0.5f * (float)SizeX;
+		float Y = acosf(NormalizedDirectionFloat.Y)*invPI * (float)SizeY;
 
 		return LookupFiltered(X, Y);
 	}
@@ -1983,8 +1984,8 @@ struct FImageViewLongLat
 	{
 		// this does the math in doubles then stores to floats :
 		//	that was probably a mistake, but leave it to avoid patches
-		float X = (1 + atan2(static_cast<float>(NormalizedDirection.X), static_cast<float>(-NormalizedDirection.Z)) / PI) / 2 * SizeX;
-		float Y = acos(static_cast<float>(NormalizedDirection.Y)) / PI * SizeY;
+		float X = (float)((1 + atan2(NormalizedDirection.X, -NormalizedDirection.Z) / PI) / 2 * (float)SizeX);
+		float Y = (float)(acos(NormalizedDirection.Y) / PI * (float)SizeY);
 
 		return LookupFiltered(X, Y);
 	}
@@ -2039,7 +2040,7 @@ static inline FVector TransformWorldToSideSpace(uint32 CubemapFace, const FVecto
 static inline FVector ComputeSSCubeDirectionAtTexelCenter(uint32 x, uint32 y, float InvSideExtent)
 {
 	// center of the texels
-	FVector DirectionSS((x + 0.5f) * InvSideExtent * 2 - 1, (y + 0.5f) * InvSideExtent * 2 - 1, 1);
+	FVector DirectionSS(((float)x + 0.5f) * InvSideExtent * 2 - 1, ((float)y + 0.5f) * InvSideExtent * 2 - 1, 1);
 	DirectionSS.Normalize();
 	return DirectionSS;
 }
@@ -2065,7 +2066,7 @@ void ITextureCompressorModule::GenerateBaseCubeMipFromLongitudeLatitude2D(FImage
 
 	// TODO_TEXTURE: Expose target size to user.
 	uint32 Extent = ComputeLongLatCubemapExtents(LongLatImage.SizeX, MaxCubemapTextureResolution);
-	float InvExtent = 1.0f / Extent;
+	float InvExtent = 1.0f / (float)Extent;
 	OutMip->Init(Extent, Extent, SrcImage.NumSlices * 6, ERawImageFormat::RGBA32F, EGammaSpace::Linear);
 
 	for (int32 Slice = 0; Slice < SrcImage.NumSlices; ++Slice)
@@ -2105,8 +2106,8 @@ public:
 
 		// *2 as the position is from -1 to 1
 		// / InFullExtent as x and y is in the range 0..InFullExtent-1
-		PositionToWorldScale = 2.0f / InFullExtent;
-		InvFullExtent = 1.0f / FullExtent;
+		PositionToWorldScale = 2.0f / (float)InFullExtent;
+		InvFullExtent = 1.0f / (float)FullExtent;
 
 		// examples: 0 to diffuse convolution, 0.95f for glossy
 		DirDot = FMath::Min(FMath::Cos(ConeAngle), 0.9999f);
@@ -2121,11 +2122,11 @@ public:
 	// @return true: yes, traverse deeper, false: not relevant
 	bool TestIfRelevant(uint32 x, uint32 y, uint32 LocalExtent) const
 	{
-		float HalfExtent = LocalExtent * 0.5f; 
-		float U = (x + HalfExtent) * PositionToWorldScale - 1.0f;
-		float V = (y + HalfExtent) * PositionToWorldScale - 1.0f;
+		float HalfExtent = (float)LocalExtent * 0.5f; 
+		float U = ((float)x + HalfExtent) * PositionToWorldScale - 1.0f;
+		float V = ((float)y + HalfExtent) * PositionToWorldScale - 1.0f;
 
-		float SphereRadius = RadiusToWorldScale * LocalExtent;
+		float SphereRadius = RadiusToWorldScale * (float)LocalExtent;
 
 		FVector SpherePos(U, V, 1);
 
@@ -2248,8 +2249,8 @@ static inline float TriangleArea2_3D(FVector A, FVector B, FVector C)
 
 static inline float ComputeTexelArea(uint32 x, uint32 y, float InvSideExtentMul2)
 {
-	float fU = x * InvSideExtentMul2 - 1;
-	float fV = y * InvSideExtentMul2 - 1;
+	float fU = (float)x * InvSideExtentMul2 - 1;
+	float fV = (float)y * InvSideExtentMul2 - 1;
 
 	FVector CornerA = FVector(fU, fV, 1);
 	FVector CornerB = FVector(fU + InvSideExtentMul2, fV, 1);
@@ -2275,7 +2276,7 @@ static void GenerateAngularFilteredMip(FImage* DestMip, FImage& SrcMip, float Co
 	TRACE_CPUPROFILER_EVENT_SCOPE(Texture.GenerateAngularFilteredMip);
 
 	int32 MipExtent = DestMip->SizeX;
-	float MipInvSideExtent = 1.0f / MipExtent;
+	float MipInvSideExtent = 1.0f / (float)MipExtent;
 
 	TArray<float> TexelAreaArray;
 	TexelAreaArray.AddUninitialized(SrcMip.SizeX * SrcMip.SizeY);
@@ -2313,7 +2314,7 @@ static void GenerateAngularFilteredMip(FImage* DestMip, FImage& SrcMip, float Co
 
 			void DoWork()
 			{
-				const float InvSideExtent = 1.0f / Extent;
+				const float InvSideExtent = 1.0f / (float)Extent;
 				FImageView2D DestMipView(*DestMip, Face);
 				for (int32 y = 0; y < Extent; ++y)
 				{
@@ -2410,9 +2411,9 @@ void ITextureCompressorModule::GenerateAngularFilteredMips(TArray<FImage>& InOut
 	for (int32 i = 0; i < NumMips; ++i)
 	{
 		// 0:top mip 1:lowest mip = diffuse convolve
-		float NormalizedMipLevel = i / (float)(NumMips - DiffuseConvolveMipLevel);
-		float AdjustedMipLevel = NormalizedMipLevel * NumMips;
-		float NormalizedWidth = BaseExtent * FMath::Pow(2.0f, -AdjustedMipLevel);
+		float NormalizedMipLevel = (float)i / (float)(NumMips - DiffuseConvolveMipLevel);
+		float AdjustedMipLevel = NormalizedMipLevel * (float)NumMips;
+		float NormalizedWidth = (float)BaseExtent * FMath::Pow(2.0f, -AdjustedMipLevel);
 		float TexelSize = 1.0f / NormalizedWidth;
 
 		// 0.001f:sharp  .. PI/2: diffuse convolve
@@ -2439,7 +2440,7 @@ void ITextureCompressorModule::GenerateAngularFilteredMips(TArray<FImage>& InOut
 		// unoptimized
 		//	float FloatInputMip = FMath::Log2(FMath::Sqrt(AreaCoveredInNormalizedArea)) + InputMipCount - QualityBias;
 		// optimized
-		float FloatInputMip = 0.5f * FMath::Log2(AreaCoveredInNormalizedArea) + NumMips - QualityBias;
+		float FloatInputMip = 0.5f * FMath::Log2(AreaCoveredInNormalizedArea) + (float)NumMips - QualityBias;
 		uint32 InputMip = FMath::Clamp(FMath::TruncToInt(FloatInputMip), 0, NumMips - 1);
 
 		FImage& Mip = InOutMipChain.Emplace_GetRef(Extent, Extent, 6, ERawImageFormat::RGBA32F);
@@ -3110,15 +3111,15 @@ static bool ApplyCompositeTexture(FImage& DestRoughness, const FImage& SourceNor
 		//const FImageView2D NormalColors(const_cast<FImage &>(SourceNormals),0);
 		const float * SourceNormalLengthsPlane = &SourceNormalLengths[0];
 
-		float InvDestW = 1.f/DestRoughness.SizeX;
-		float InvDestH = 1.f/DestRoughness.SizeY;
+		float InvDestW = 1.f/(float)DestRoughness.SizeX;
+		float InvDestH = 1.f/(float)DestRoughness.SizeY;
 
 		for( int64 DestY=0; DestY< DestRoughness.SizeY; DestY++)
 		{
-			float V = (DestY + 0.5f) * InvDestH;
+			float V = ((float)DestY + 0.5f) * InvDestH;
 			for( int64 DestX=0; DestX< DestRoughness.SizeX; DestX++)
 			{
-				float U = (DestX + 0.5f) * InvDestW;
+				float U = ((float)DestX + 0.5f) * InvDestW;
 				
 				//const FLinearColor NormalColor = LookupSourceMipBilinearUV(NormalColors,U,V);
 				const float NormalLength = LookupFloatBilinearUV(SourceNormalLengthsPlane,SourceNormals.SizeX,SourceNormals.SizeY,U,V);
diff --git a/Engine/Source/Developer/TextureFormat/Public/ChildTextureFormat.h b/Engine/Source/Developer/TextureFormat/Public/ChildTextureFormat.h
--- a/Engine/Source/Developer/TextureFormat/Public/ChildTextureFormat.h
+++ b/Engine/Source/Developer/TextureFormat/Public/ChildTextureFormat.h
@@ -188,7 +188,7 @@ public:
 		uint8 ChildVersion = GetChildFormatVersion(Format, BuildSettings);
 
 		// 8 bits for each version
-		return (BaseVersion << 8) | ChildVersion;
+		return (uint16)((BaseVersion << 8) | ChildVersion);
 	}
 
 	virtual FString GetDerivedDataKeyString(const FTextureBuildSettings& InBuildSettings, int32 InMipCount, const FIntVector3& InMip0Dimensions) const final
diff --git a/Engine/Source/Editor/TextureEditor/Private/Models/TextureEditorViewportClient.cpp b/Engine/Source/Editor/TextureEditor/Private/Models/TextureEditorViewportClient.cpp
--- a/Engine/Source/Editor/TextureEditor/Private/Models/TextureEditorViewportClient.cpp
+++ b/Engine/Source/Editor/TextureEditor/Private/Models/TextureEditorViewportClient.cpp
@@ -187,8 +187,8 @@ void FTextureEditorViewportClient::Draw(FViewport* Viewport, FCanvas* Canvas)
 	// the texture has transparency
 	if (Viewport && CheckerboardTexture)
 	{
-		const int32 CheckerboardSizeX = FMath::Max<int32>(1, CheckerboardTexture->GetSizeX());
-		const int32 CheckerboardSizeY = FMath::Max<int32>(1, CheckerboardTexture->GetSizeY());
+		const float CheckerboardSizeX = (float)FMath::Max<int32>(1, CheckerboardTexture->GetSizeX());
+		const float CheckerboardSizeY = (float)FMath::Max<int32>(1, CheckerboardTexture->GetSizeY());
 		if (Settings.Background == TextureEditorBackground_CheckeredFill)
 		{
 			Canvas->DrawTile( 0.0f, 0.0f, Viewport->GetSizeXY().X, Viewport->GetSizeXY().Y, 0.0f, 0.0f, (float)Viewport->GetSizeXY().X / CheckerboardSizeX, (float)Viewport->GetSizeXY().Y / CheckerboardSizeY, FLinearColor::White, CheckerboardTexture->GetResource());
@@ -221,7 +221,8 @@ void FTextureEditorViewportClient::Draw(FViewport* Viewport, FCanvas* Canvas)
 		// Draw a white border around the texture to show its extents
 		if (Settings.TextureBorderEnabled)
 		{
-			FCanvasBoxItem BoxItem(FVector2D(XPos - (BorderSize - 1) * 0.5f, YPos - (BorderSize - 1) * 0.5f), FVector2D(Width + BorderSize, Height + BorderSize));
+			float ScaledBorderSize = ((float)BorderSize - 1) * 0.5f;
+			FCanvasBoxItem BoxItem(FVector2D((float)XPos - ScaledBorderSize, (float)YPos - ScaledBorderSize), FVector2D(Width + BorderSize, Height + BorderSize));
 			BoxItem.LineThickness = (float)BorderSize;
 			BoxItem.SetColor( Settings.TextureBorderColor );
 			Canvas->DrawItem( BoxItem );
@@ -438,7 +439,7 @@ float FTextureEditorViewportClient::GetViewportVerticalScrollBarRatio() const
 		WidgetHeight = (float)(TextureEditorViewportPtr.Pin()->GetViewport()->GetSizeXY().Y);
 	}
 
-	return WidgetHeight / Height;
+	return WidgetHeight / (float)Height;
 }
 
 
@@ -455,7 +456,7 @@ float FTextureEditorViewportClient::GetViewportHorizontalScrollBarRatio() const
 		WidgetWidth = (float)(TextureEditorViewportPtr.Pin()->GetViewport()->GetSizeXY().X);
 	}
 
-	return WidgetWidth / Width;
+	return WidgetWidth / (float)Width;
 }
 
 
@@ -519,7 +520,7 @@ FVector2D FTextureEditorViewportClient::GetViewportScrollBarPositions() const
 
 		if ((TextureEditorViewportPtr.Pin()->GetVerticalScrollBar()->GetVisibility() == EVisibility::Visible) && VDistFromBottom < 1.0f)
 		{
-			Positions.Y = FMath::Clamp((1.0f + VDistFromTop - VDistFromBottom - VRatio) * 0.5f, 0.0f, 1.0f - VRatio) * Height;
+			Positions.Y = FMath::Clamp((1.0f + VDistFromTop - VDistFromBottom - VRatio) * 0.5f, 0.0f, 1.0f - VRatio) * (float)Height;
 		}
 		else
 		{
@@ -528,7 +529,7 @@ FVector2D FTextureEditorViewportClient::GetViewportScrollBarPositions() const
 
 		if ((TextureEditorViewportPtr.Pin()->GetHorizontalScrollBar()->GetVisibility() == EVisibility::Visible) && HDistFromBottom < 1.0f)
 		{
-			Positions.X = FMath::Clamp((1.0f + HDistFromTop - HDistFromBottom - HRatio) * 0.5f, 0.0f, 1.0f - HRatio) * Width;
+			Positions.X = FMath::Clamp((1.0f + HDistFromTop - HDistFromBottom - HRatio) * 0.5f, 0.0f, 1.0f - HRatio) * (float)Width;
 		}
 		else
 		{
diff --git a/Engine/Source/Editor/TextureEditor/Private/TextureEditorToolkit.cpp b/Engine/Source/Editor/TextureEditor/Private/TextureEditorToolkit.cpp
--- a/Engine/Source/Editor/TextureEditor/Private/TextureEditorToolkit.cpp
+++ b/Engine/Source/Editor/TextureEditor/Private/TextureEditorToolkit.cpp
@@ -910,7 +910,7 @@ void FTextureEditorToolkit::SetCustomZoomLevel( double ZoomValue )
 	// the mouse wheel zoom is quantized on ZoomFactorLogSteps
 	//	but that's too chunky for the drag slider, give it more steps, but on the same quantization grid
 	double QuantizationSteps = ZoomFactorLogSteps*2.0;
-	double LogZoomQuantized = (1.0/QuantizationSteps) * FMath::RoundToInt( QuantizationSteps * LogZoom );
+	double LogZoomQuantized = (1.0/QuantizationSteps) * (double)FMath::RoundToInt( QuantizationSteps * LogZoom );
 	ZoomValue = pow(2.0,LogZoomQuantized);
 
 	ZoomValue = FMath::Clamp(ZoomValue, MinZoom, MaxZoom);
diff --git a/Engine/Source/Runtime/Engine/Public/CanvasTypes.h b/Engine/Source/Runtime/Engine/Public/CanvasTypes.h
--- a/Engine/Source/Runtime/Engine/Public/CanvasTypes.h
+++ b/Engine/Source/Runtime/Engine/Public/CanvasTypes.h
@@ -880,7 +880,7 @@ struct FScreenMessageWriter
 		Y += 14;
 	}
 
-	ENGINE_API void DrawLine(const FText& Message, int32 X = 10, const FLinearColor& Color = FLinearColor(1.0, 0.05, 0.05, 1.0));
+	ENGINE_API void DrawLine(const FText& Message, int32 X = 10, const FLinearColor& Color = FLinearColor(1.0f, 0.05f, 0.05f, 1.0f));
 
 	FCanvas& Canvas;
 	int32 Y;
diff --git a/Engine/Source/Runtime/Slate/Public/Widgets/Input/SSpinBox.h b/Engine/Source/Runtime/Slate/Public/Widgets/Input/SSpinBox.h
--- a/Engine/Source/Runtime/Slate/Public/Widgets/Input/SSpinBox.h
+++ b/Engine/Source/Runtime/Slate/Public/Widgets/Input/SSpinBox.h
@@ -569,7 +569,7 @@ public:
 				const int32 CachedShiftMouseMovePixelPerDelta = ShiftMouseMovePixelPerDelta.Get();
 				if (CachedShiftMouseMovePixelPerDelta > 1 && MouseEvent.IsShiftDown())
 				{
-					SliderWidthInSlateUnits *= CachedShiftMouseMovePixelPerDelta;
+					SliderWidthInSlateUnits *= (float)CachedShiftMouseMovePixelPerDelta;
 				}
 
 				if (MouseEvent.IsControlDown())
@@ -643,7 +643,7 @@ public:
 					const double Sign = (MouseEvent.GetCursorDelta().X > 0) ? 1.0 : -1.0;
 					if (LinearDeltaSensitivity.IsSet() && LinearDeltaSensitivity.Get() != 0 && Delta.IsSet() && Delta.Get() > 0)
 					{
-						const double MouseDelta = FMath::Abs(MouseEvent.GetCursorDelta().X / LinearDeltaSensitivity.Get());
+						const double MouseDelta = FMath::Abs(MouseEvent.GetCursorDelta().X / (float)LinearDeltaSensitivity.Get());
 						NewValue = InternalValue + (Sign * MouseDelta * FMath::Pow((double)Delta.Get(), SliderExponent.Get()));
 					}
 					else
