Reducing texture build memory usage by freeing mip pixel memory as they are no longer needed

diff --git a/Engine/Source/Developer/TextureCompressor/Private/TextureCompressorModule.cpp b/Engine/Source/Developer/TextureCompressor/Private/TextureCompressorModule.cpp
--- a/Engine/Source/Developer/TextureCompressor/Private/TextureCompressorModule.cpp
+++ b/Engine/Source/Developer/TextureCompressor/Private/TextureCompressorModule.cpp
@@ -3177,7 +3177,7 @@ void FTextureBuildSettings::GetEncodedTextureDescription(FEncodedTextureDescript
 // compress mip-maps in InMipChain and add mips to Texture, might alter the source content
 static bool CompressMipChain(
 	const ITextureFormat* TextureFormat,
-	const TArray<FImage>& MipChain,
+	TArray<FImage>& MipChain,
 	const FTextureBuildSettings& Settings,
 	const bool bImageHasAlphaChannel,
 	FStringView DebugTexturePathName,
@@ -3262,17 +3262,7 @@ static bool CompressMipChain(
 				OutMips[MipIndex]
 			);
 
-			// no longer possible, because hashing is done on a thread
-			// todo: restore this by making the images refcounted to track lifetime better
-			/*
-			// note: MipChain[MipIndex].RawData may be freed or mutated by CompressImage
-			// do not use it after the call to CompressImage
-			// go ahead and free it now if CompressImage didn't :
-			for(int MipSubIndex=0;MipSubIndex<MipsToCompress;MipSubIndex++)
-			{
-				MipChain[MipIndex+MipSubIndex].RawData.Empty();
-			}
-			*/
+			MipChain[MipIndex].RawData.Empty();
 		}
 
 		return bSuccess;
@@ -3544,7 +3534,7 @@ public:
 	// compute a hash used to identify the contents of a mip chain
 	// this is used by bulkdata diff tool, stored in asset registry
 	// it is for debug tools only and skipping it is optional
-	static uint64 ComputeMipChainHash(const TArray<FImage> & IntermediateMipChain)
+	static uint64 ComputeMipChainHash(TArray<FImage>& IntermediateMipChain)
 	{
 		TRACE_CPUPROFILER_EVENT_SCOPE(Texture.ComputeMipChainHash);
 
@@ -3695,8 +3685,6 @@ public:
 		
 		UE::Tasks::TTask<uint64> HashingTask;
 		TArray<FImageInfo> SaveImageInfos;
-		
-		//FIoHash DebugOnly_HashBefore = ComputeMipChainHash(IntermediateMipChain);
 
 		if (OutMetadata)
 		{
@@ -3714,24 +3702,14 @@ public:
 			// back out so that when textures happen to get saved, it goes with it.
 			OutMetadata->bSourceMipsAlphaDetected = bSourceMipsAlphaDetected;
 
-			// hash value is not needed immediately, so don't block on it :
-			HashingTask = UE::Tasks::Launch(TEXT("Texture.MipHashBuilder.Task"), [&](){ return ComputeMipChainHash(IntermediateMipChain); } );
+			OutMetadata->PreEncodeMipsHash = ComputeMipChainHash(IntermediateMipChain);
 		}
 		
 		bool bCompressSucceeded = CompressMipChain(TextureFormat, IntermediateMipChain, BuildSettings, bImageHasAlphaChannel, DebugTexturePathName,
 					OutTextureMips, OutNumMipsInTail, OutExtData);
-		
-		// make sure IntermediateMipChain isn't changed:
-		//FIoHash DebugOnly_HashAfter = ComputeMipChainHash(IntermediateMipChain);
-		//check( DebugOnly_HashBefore == DebugOnly_HashAfter );
 
 		if (OutMetadata)
 		{
-			// now block on the hash task
-			OutMetadata->PreEncodeMipsHash = HashingTask.GetResult();
-		
-			//check( DebugOnly_HashBefore == OutMetadata->PreEncodeMipsHash );
-		
 			// (rough) check that IntermediateMipChain was not modified by the TextureFormats :
 			check( SaveImageInfos.Num() == IntermediateMipChain.Num() );
 			for (int32 i=0;i<IntermediateMipChain.Num();++i)
@@ -3963,7 +3941,12 @@ private:
 
 				// Max Texture Size resizing happens here :
 				// note we do not check for TMGS_Angular here
-				GenerateMipChain(BuildSettings, bSuitableFormat ? BaseImage : Temp, BuildSourceImageMips, 1);
+				const FImage& BaseMip = bSuitableFormat ? BaseImage : Temp;
+				GenerateMipChain(BuildSettings, BaseMip, BuildSourceImageMips, 1);
+
+				// mip data not needed anymore
+				// todo: this could free "BaseMip" image instead, if it's ok with caller (currently const type used)
+				Temp.RawData.Empty();
 
 				while( BuildSourceImageMips.Last().SizeX > MaxTextureResolution || 
 					   BuildSourceImageMips.Last().SizeY > MaxTextureResolution )
@@ -3971,7 +3954,11 @@ private:
 					// note: now making mips one by one, rather than N in one call
 					//	this is not exactly the same if AlphaCoverage processing is on
 					check( BuildSourceImageMips.Num() < BuildSourceImageMipsMaxCount );
-					GenerateMipChain(BuildSettings, BuildSourceImageMips.Last(), BuildSourceImageMips, 1);
+					FImage& LastMip = BuildSourceImageMips.Last();
+					GenerateMipChain(BuildSettings, LastMip, BuildSourceImageMips, 1);
+
+					// mip data not needed anymore
+					LastMip.RawData.Empty();
 				}
 			
 				check( BuildSourceImageMips.Last().SizeX <= MaxTextureResolution &&
