ColorManagement: Utilize OpenColorIO for advanced color space transformations on textures during mip generation.

Relies on CL 26098047.

diff --git a/Engine/Source/Developer/TextureCompressor/Private/TextureCompressorModule.cpp b/Engine/Source/Developer/TextureCompressor/Private/TextureCompressorModule.cpp
--- a/Engine/Source/Developer/TextureCompressor/Private/TextureCompressorModule.cpp
+++ b/Engine/Source/Developer/TextureCompressor/Private/TextureCompressorModule.cpp
@@ -16,6 +16,7 @@
 #include "Interfaces/ITextureFormat.h"
 #include "Misc/Paths.h"
 #include "Tasks/Task.h"
+#include "OpenColorIOWrapper.h"
 #include "ImageCore.h"
 #include "ImageParallelFor.h"
 #include <cmath>
@@ -1598,6 +1599,55 @@ static void DownscaleImage(const FImage& SrcImage, FImage& DstImage, const FText
 	}
 }
 
+// Replaces previous calls to FImage::Linearize and FImageCore::TransformToWorkingColorSpace
+static void LinearizeToWorkingColorSpace(const FImage& SrcImage, FImage& DstImage, const FTextureBuildSettings& BuildSettings)
+{
+	TRACE_CPUPROFILER_EVENT_SCOPE(Texture.LinearizeToWorkingColorSpace);
+
+	// Allocate the 32-bit linear floating-point image
+	DstImage.Init(SrcImage.SizeX, SrcImage.SizeY, SrcImage.NumSlices, ERawImageFormat::RGBA32F, EGammaSpace::Linear);
+
+	FImageView SrcImageView(SrcImage);
+	const UE::Color::EEncoding SourceEncodingOverride = static_cast<UE::Color::EEncoding>(BuildSettings.SourceEncodingOverride);
+
+	// If the source encoding override is active, we avoid CopyImage's de-gammatization and instead let OpenColorIO do the decoding below.
+	if (SourceEncodingOverride != UE::Color::EEncoding::None)
+	{
+		SrcImageView.GammaSpace = DstImage.GammaSpace; // EGammaSpace::Linear
+	}
+
+	FImageCore::CopyImage(SrcImageView, DstImage);
+
+	// Decode and/or color transform to the working color space when needed
+	if (SourceEncodingOverride != UE::Color::EEncoding::None || BuildSettings.bHasColorSpaceDefinition)
+	{
+		FOpenColorIOWrapperSourceColorSettings SrcSettings;
+		SrcSettings.EncodingOverride = SourceEncodingOverride;
+
+		if (BuildSettings.bHasColorSpaceDefinition)
+		{
+			// Matrix transform will be applied by OpenColorIO, unless it is already equal to the working color space.
+			// Note: We no longer manually apply SaturateToHalfFloat(..) as this is now handled internally by the library.
+			TStaticArray<FVector2d, 4> SourceChromaticities;
+			SourceChromaticities[0] = FVector2d(BuildSettings.RedChromaticityCoordinate);
+			SourceChromaticities[1] = FVector2d(BuildSettings.GreenChromaticityCoordinate);
+			SourceChromaticities[2] = FVector2d(BuildSettings.BlueChromaticityCoordinate);
+			SourceChromaticities[3] = FVector2d(BuildSettings.WhiteChromaticityCoordinate);
+
+			SrcSettings.ColorSpaceOverride = MoveTemp(SourceChromaticities);
+			SrcSettings.ChromaticAdaptationMethod = static_cast<UE::Color::EChromaticAdaptationMethod>(BuildSettings.ChromaticAdaptationMethod);
+		}
+
+		// Invoke the OpenColorIO library processor with the specified transformation to the working color space.
+		const FOpenColorIOWrapperProcessor Processor = FOpenColorIOWrapperProcessor::CreateTransformToWorkingColorSpace(SrcSettings);
+		const bool bSuccess = Processor.TransformImage(DstImage);
+		if (!bSuccess)
+		{
+			UE_LOG(LogTextureCompressor, Error, TEXT("Failed to linearize the texture to the working color space."));
+		}
+	}
+}
+
 void ITextureCompressorModule::GenerateMipChain(
 	const FTextureBuildSettings& Settings,
 	const FImage& BaseImage,
@@ -2059,14 +2109,24 @@ static uint32 ComputeLongLatCubemapExtents(int32 SrcImageSizeX, const uint32 Max
 }
 
 void ITextureCompressorModule::GenerateBaseCubeMipFromLongitudeLatitude2D(FImage* OutMip, const FImage& SrcImage, const uint32 MaxCubemapTextureResolution, uint8 SourceEncodingOverride)
+{
+	FTextureBuildSettings TempBuildSettings;
+	TempBuildSettings.MaxTextureResolution = MaxCubemapTextureResolution;
+	TempBuildSettings.SourceEncodingOverride = SourceEncodingOverride;
+	TempBuildSettings.bHasColorSpaceDefinition = false;
+
+	GenerateBaseCubeMipFromLongitudeLatitude2D(OutMip, SrcImage, TempBuildSettings);
+}
+
+void ITextureCompressorModule::GenerateBaseCubeMipFromLongitudeLatitude2D(FImage* OutMip, const FImage& SrcImage, const FTextureBuildSettings& InBuildSettings)
 {
 	TRACE_CPUPROFILER_EVENT_SCOPE(Texture.CubeMipFromLongLat);
 
 	FImage LongLatImage;
-	SrcImage.Linearize(SourceEncodingOverride, LongLatImage);
+	LinearizeToWorkingColorSpace(SrcImage, LongLatImage, InBuildSettings);
 
 	// TODO_TEXTURE: Expose target size to user.
-	uint32 Extent = ComputeLongLatCubemapExtents(LongLatImage.SizeX, MaxCubemapTextureResolution);
+	uint32 Extent = ComputeLongLatCubemapExtents(LongLatImage.SizeX, InBuildSettings.MaxTextureResolution);
 	float InvExtent = 1.0f / (float)Extent;
 	OutMip->Init(Extent, Extent, SrcImage.NumSlices * 6, ERawImageFormat::RGBA32F, EGammaSpace::Linear);
 
@@ -4024,7 +4084,7 @@ private:
 			if (bLongLatCubemap)
 			{
 				// Generate the base mip from the long-lat source image.
-				GenerateBaseCubeMipFromLongitudeLatitude2D(&Mip, Image, BuildSettings.MaxTextureResolution, BuildSettings.SourceEncodingOverride);
+				GenerateBaseCubeMipFromLongitudeLatitude2D(&Mip, Image, BuildSettings);
 	
 				check( CopyCount == 1 );
 			}
@@ -4034,7 +4094,7 @@ private:
 				if(BuildSettings.bApplyKernelToTopMip)
 				{
 					FImage Temp;
-					Image.Linearize(BuildSettings.SourceEncodingOverride, Temp);
+					LinearizeToWorkingColorSpace(Image, Temp, BuildSettings);
 					if(BuildSettings.bRenormalizeTopMip)
 					{
 						NormalizeMip(Temp);
@@ -4046,7 +4106,7 @@ private:
 				{
 					if (bLinearize)
 					{
-						Image.Linearize(BuildSettings.SourceEncodingOverride, Mip);
+						LinearizeToWorkingColorSpace(Image, Mip, BuildSettings);
 						if (BuildSettings.bRenormalizeTopMip)
 						{
 							NormalizeMip(Mip);
@@ -4078,17 +4138,6 @@ private:
 				DownscaleImage(Mip, Mip, FTextureDownscaleSettings(BuildSettings));
 			}
 
-			if (BuildSettings.bHasColorSpaceDefinition)
-			{
-				FImageCore::TransformToWorkingColorSpace(
-					Mip,
-					FVector2d(BuildSettings.RedChromaticityCoordinate),
-					FVector2d(BuildSettings.GreenChromaticityCoordinate),
-					FVector2d(BuildSettings.BlueChromaticityCoordinate),
-					FVector2d(BuildSettings.WhiteChromaticityCoordinate),
-					static_cast<UE::Color::EChromaticAdaptationMethod>(BuildSettings.ChromaticAdaptationMethod));
-			}
-
 			// Apply color adjustments
 			AdjustImageColors(Mip, BuildSettings);
 
diff --git a/Engine/Source/Developer/TextureCompressor/Public/TextureCompressorModule.h b/Engine/Source/Developer/TextureCompressor/Public/TextureCompressorModule.h
--- a/Engine/Source/Developer/TextureCompressor/Public/TextureCompressorModule.h
+++ b/Engine/Source/Developer/TextureCompressor/Public/TextureCompressorModule.h
@@ -407,8 +407,16 @@ public:
 	 * @param OutMip - The output mip.
 	 * @param SrcImage - The source longlat image.
 	 */
+	UE_DEPRECATED(5.3, "GenerateBaseCubeMipFromLongitudeLatitude2D with FTextureBuildSettings should be used.")
 	TEXTURECOMPRESSOR_API static void GenerateBaseCubeMipFromLongitudeLatitude2D(FImage* OutMip, const FImage& SrcImage, const uint32 MaxCubemapTextureResolution, uint8 SourceEncodingOverride = 0);
 
+	/**
+	 * Generates the base cubemap mip from a longitude-latitude 2D image.
+	 * @param OutMip - The output mip.
+	 * @param SrcImage - The source longlat image.
+	 * @param InBuildSettings - Image build settings
+	 */
+	TEXTURECOMPRESSOR_API static void GenerateBaseCubeMipFromLongitudeLatitude2D(FImage* OutMip, const FImage& SrcImage, const FTextureBuildSettings& InBuildSettings);
 
 	/**
 	 * Generates angularly filtered mips.
diff --git a/Engine/Source/Developer/TextureCompressor/TextureCompressor.Build.cs b/Engine/Source/Developer/TextureCompressor/TextureCompressor.Build.cs
--- a/Engine/Source/Developer/TextureCompressor/TextureCompressor.Build.cs
+++ b/Engine/Source/Developer/TextureCompressor/TextureCompressor.Build.cs
@@ -15,6 +15,7 @@ public class TextureCompressor : ModuleRules
 			new string[] {
 				"Core",
 				"ImageCore",
+				"OpenColorIOWrapper",
 				"TextureBuildUtilities",
 				"TextureFormat",
 			}
diff --git a/Engine/Source/Runtime/ColorManagement/Public/ColorManagementDefines.h b/Engine/Source/Runtime/ColorManagement/Public/ColorManagementDefines.h
--- a/Engine/Source/Runtime/ColorManagement/Public/ColorManagementDefines.h
+++ b/Engine/Source/Runtime/ColorManagement/Public/ColorManagementDefines.h
@@ -6,11 +6,11 @@
 
 namespace UE { namespace Color {
 
-/** Increment upon breaking changes to the EEncoding enum. */
-constexpr uint32 ENCODING_TYPES_VER = 3;
+/** Increment upon breaking changes to the EEncoding enum. Note that changing this forces a rebuild of textures that rely on it. */
+constexpr uint32 ENCODING_TYPES_VER = 4;
 
-/** Increment upon breaking changes to the EColorSpace and EChromaticAdaptationMethod enums. */
-constexpr uint32 COLORSPACE_VER = 1;
+/** Increment upon breaking changes to the EColorSpace and EChromaticAdaptationMethod enums. Note that changing this forces a rebuild of textures that rely on it.*/
+constexpr uint32 COLORSPACE_VER = 2;
 
 /** List of available encodings/transfer functions.
 * 
diff --git a/Engine/Source/Runtime/Engine/Private/TextureDerivedDataBuildUtils.cpp b/Engine/Source/Runtime/Engine/Private/TextureDerivedDataBuildUtils.cpp
--- a/Engine/Source/Runtime/Engine/Private/TextureDerivedDataBuildUtils.cpp
+++ b/Engine/Source/Runtime/Engine/Private/TextureDerivedDataBuildUtils.cpp
@@ -197,6 +197,14 @@ static void WriteBuildSettings(FCbWriter& Writer, const FTextureBuildSettings& B
 	{
 		WriteCbField<bool>(Writer, "VTPow22_ForceNewDDcKey", true); 
 	}
+	// Force a new DDC key when utilizing color space transforms due to the switch to
+	// OpenColorIO. Remove when the build version is updated and all textures get rebuilt
+	// with it.
+	if (BuildSettings.SourceEncodingOverride != 0 /*UE::Color::EEncoding::None*/ ||
+		BuildSettings.bHasColorSpaceDefinition)
+	{
+		WriteCbField<bool>(Writer, "OCIO_ForceNewKey", true);
+	}
 
 	Writer.EndObject();
 }
