Add Stretch to Power Of Two mode in Texture Settings.

diff --git a/Engine/Source/Developer/TextureCompressor/Private/TextureCompressorModule.cpp b/Engine/Source/Developer/TextureCompressor/Private/TextureCompressorModule.cpp
--- a/Engine/Source/Developer/TextureCompressor/Private/TextureCompressorModule.cpp
+++ b/Engine/Source/Developer/TextureCompressor/Private/TextureCompressorModule.cpp
@@ -3427,12 +3427,14 @@ static bool GetPowerOfTwoTargetTextureSize(int32 InMip0SizeX, int32 InMip0SizeY,
 	// None should not get here
 
 	case ETexturePowerOfTwoSetting::PadToPowerOfTwo:
+	case ETexturePowerOfTwoSetting::StretchToPowerOfTwo:
 		TargetTextureSizeX = PowerOfTwoTextureSizeX;
 		TargetTextureSizeY = PowerOfTwoTextureSizeY;
 		TargetTextureSizeZ = PowerOfTwoTextureSizeZ;
 		break;
 
 	case ETexturePowerOfTwoSetting::PadToSquarePowerOfTwo:
+	case ETexturePowerOfTwoSetting::StretchToSquarePowerOfTwo:
 		TargetTextureSizeX = TargetTextureSizeY = TargetTextureSizeZ =
 			FMath::Max3<int32>(PowerOfTwoTextureSizeX, PowerOfTwoTextureSizeY, PowerOfTwoTextureSizeZ);
 		break;
@@ -3874,15 +3876,23 @@ private:
 
 			if (bPadOrStretchTexture)
 			{
+				bool bStretchTexture = PowerOfTwoMode == ETexturePowerOfTwoSetting::StretchToPowerOfTwo || PowerOfTwoMode == ETexturePowerOfTwoSetting::StretchToSquarePowerOfTwo;
 				if (BuildSettings.MipGenSettings == TMGS_LeaveExistingMips)
 				{
-					// pad+leave existing is broken
-					UE_LOG(LogTextureCompressor, Error,	TEXT("Texture padded to pow2 + LeaveExistingMips forbidden"));
+					// pad/stretch+leave existing is broken
+					UE_LOG(LogTextureCompressor, Error,	TEXT("Texture padding or stretching to power of two is not allowed when leaving existing mips."));
 					return false;
 				}
 				if ( bLongLatCubemap )
 				{
-					UE_LOG(LogTextureCompressor, Warning, TEXT("PadPow2 + LongLat cubemap doesn't work, continuing.."));
+					if (bStretchTexture)
+					{
+						UE_LOG(LogTextureCompressor, Warning, TEXT("In order to improve the quality of the generated texture, stretching LongLat cubemaps should be avoided."));
+					}
+					else
+					{
+						UE_LOG(LogTextureCompressor, Warning, TEXT("Padding of a LongLat cubemap may result in incorrect mapping of the texture pixels to the cubemap faces and should be avoided."));
+					}
 				}
 
 				// Want to stretch or pad the texture
@@ -3898,44 +3908,57 @@ private:
 				// space for one source mip and one destination mip
 				const FImage& SourceImage = bSuitableFormat ? FirstSourceMipImage : Temp;
 				FImage& TargetImage = PaddedSourceMips.Emplace_GetRef(TargetTextureSizeX, TargetTextureSizeY, BuildSettings.bVolume ? TargetTextureSizeZ : SourceImage.NumSlices, SourceImage.Format);
-				FLinearColor FillColor = BuildSettings.PaddingColor;
-
-				FLinearColor* TargetPtr = (FLinearColor*)TargetImage.RawData.GetData();
-				FLinearColor* SourcePtr = (FLinearColor*)SourceImage.RawData.GetData();
-				check(SourceImage.GetBytesPerPixel() == sizeof(FLinearColor));
-				check(TargetImage.GetBytesPerPixel() == sizeof(FLinearColor));
 
-				for (int32 SliceIndex = 0; SliceIndex < SourceImage.NumSlices; ++SliceIndex)
+				if (bStretchTexture)
 				{
-					for (int32 Y = 0; Y < TargetTextureSizeY; ++Y)
+					if (TargetImage.NumSlices != 1)
 					{
-						int32 XStart = 0;
-						if (Y < SourceImage.SizeY)
-						{
-							XStart = SourceImage.SizeX;
-							FMemory::Memcpy(TargetPtr, SourcePtr, SourceImage.SizeX * sizeof(FLinearColor));
-							SourcePtr += SourceImage.SizeX;
-							TargetPtr += SourceImage.SizeX;
-						}
+						// FImageCore::ResizeTo currently only supports resizing textures with 1 slice
+						UE_LOG(LogTextureCompressor, Error, TEXT("Texture stretching is currently only supported on Texture2D."));
+						return false;
+					}
+					SourceImage.ResizeTo(TargetImage, TargetTextureSizeX, TargetTextureSizeY, SourceImage.Format, SourceImage.GetGammaSpace());
+				}
+				else
+				{
+					FLinearColor FillColor = BuildSettings.PaddingColor;
 
-						for (int32 XPad = XStart; XPad < TargetImage.SizeX; ++XPad)
+					FLinearColor* TargetPtr = (FLinearColor*)TargetImage.RawData.GetData();
+					FLinearColor* SourcePtr = (FLinearColor*)SourceImage.RawData.GetData();
+					check(SourceImage.GetBytesPerPixel() == sizeof(FLinearColor));
+					check(TargetImage.GetBytesPerPixel() == sizeof(FLinearColor));
+
+					for (int32 SliceIndex = 0; SliceIndex < SourceImage.NumSlices; ++SliceIndex)
+					{
+						for (int32 Y = 0; Y < TargetTextureSizeY; ++Y)
 						{
-							*TargetPtr++ = FillColor;
+							int32 XStart = 0;
+							if (Y < SourceImage.SizeY)
+							{
+								XStart = SourceImage.SizeX;
+								FMemory::Memcpy(TargetPtr, SourcePtr, SourceImage.SizeX * sizeof(FLinearColor));
+								SourcePtr += SourceImage.SizeX;
+								TargetPtr += SourceImage.SizeX;
+							}
+
+							for (int32 XPad = XStart; XPad < TargetImage.SizeX; ++XPad)
+							{
+								*TargetPtr++ = FillColor;
+							}
 						}
 					}
-				}
-				// Pad new slices for volume texture
-				for (int32 SliceIndex = SourceImage.NumSlices; SliceIndex < TargetImage.NumSlices; ++SliceIndex)
-				{
-					for (int32 Y = 0; Y < TargetImage.SizeY; ++Y)
+					// Pad new slices for volume texture
+					for (int32 SliceIndex = SourceImage.NumSlices; SliceIndex < TargetImage.NumSlices; ++SliceIndex)
 					{
-						for (int32 X = 0; X< TargetImage.SizeX; ++X)
+						for (int32 Y = 0; Y < TargetImage.SizeY; ++Y)
 						{
-							*TargetPtr++ = FillColor;
+							for (int32 X = 0; X< TargetImage.SizeX; ++X)
+							{
+								*TargetPtr++ = FillColor;
+							}
 						}
 					}
 				}
-				
 				// change pSourceMips to point at the one padded image we made
 				pSourceMips = &PaddedSourceMips;
 			}
diff --git a/Engine/Source/Runtime/Engine/Classes/Engine/TextureDefines.h b/Engine/Source/Runtime/Engine/Classes/Engine/TextureDefines.h
--- a/Engine/Source/Runtime/Engine/Classes/Engine/TextureDefines.h
+++ b/Engine/Source/Runtime/Engine/Classes/Engine/TextureDefines.h
@@ -168,7 +168,13 @@ namespace ETexturePowerOfTwoSetting
 		PadToPowerOfTwo,
 
 		/** Pad the texture to the nearest square power of two size. */
-		PadToSquarePowerOfTwo
+		PadToSquarePowerOfTwo,
+
+		/** Stretch the texture to the nearest power of two size. */
+		StretchToPowerOfTwo,
+
+		/** Stretch the texture to the nearest square power of two size. */
+		StretchToSquarePowerOfTwo
 
 		// Note: These are serialized as as raw values in the texture DDC key, so additional entries
 		// should be added at the bottom; reordering or removing entries will require changing the GUID
diff --git a/Engine/Source/Runtime/Engine/Private/Texture.cpp b/Engine/Source/Runtime/Engine/Private/Texture.cpp
--- a/Engine/Source/Runtime/Engine/Private/Texture.cpp
+++ b/Engine/Source/Runtime/Engine/Private/Texture.cpp
@@ -526,8 +526,16 @@ void UTexture::ValidateSettingsAfterImportOrEdit(bool * pRequiresNotifyMaterials
 	{
 		if ( MipGenSettings == TMGS_LeaveExistingMips && PowerOfTwoMode != ETexturePowerOfTwoSetting::None )
 		{
-			// power of 2 pads not allowed with LeaveExistingMips
-			UE_LOG(LogTexture, Display, TEXT("Power of 2 padding cannot be used with LeaveExistingMips, disabled. (%s)"), *GetName());
+			// power of 2 padding or stretching is not allowed with LeaveExistingMips
+			UE_LOG(LogTexture, Display, TEXT("Power of 2 padding or stretching cannot be used with LeaveExistingMips, disabled. (%s)"), *GetName());
+
+			PowerOfTwoMode = ETexturePowerOfTwoSetting::None;
+		}
+
+		if ((PowerOfTwoMode == ETexturePowerOfTwoSetting::StretchToPowerOfTwo || PowerOfTwoMode == ETexturePowerOfTwoSetting::StretchToSquarePowerOfTwo) && !this->IsA<UTexture2D>())
+		{
+			// currently power of 2 stretching is only supported for 2D textures, but can be implemented for other types of textures in the future
+			UE_LOG(LogTexture, Display, TEXT("Currently power of two stretching is only supported for Texture2D, forcing PowerOfTwoMode to None. (%s)"), *GetName());
 
 			PowerOfTwoMode = ETexturePowerOfTwoSetting::None;
 		}
@@ -3326,12 +3334,13 @@ void UTexture::GetBuiltTextureSize( const ITargetPlatform* TargetPlatform , int3
 	SizeX = SourceSize.X;
 	SizeY = SourceSize.Y;
 
-	if (PowerOfTwoMode == ETexturePowerOfTwoSetting::PadToPowerOfTwo || PowerOfTwoMode == ETexturePowerOfTwoSetting::PadToSquarePowerOfTwo)
+	if (PowerOfTwoMode == ETexturePowerOfTwoSetting::PadToPowerOfTwo || PowerOfTwoMode == ETexturePowerOfTwoSetting::PadToSquarePowerOfTwo ||
+		PowerOfTwoMode == ETexturePowerOfTwoSetting::StretchToPowerOfTwo || PowerOfTwoMode == ETexturePowerOfTwoSetting::StretchToSquarePowerOfTwo)
 	{
 		SizeX = FMath::RoundUpToPowerOfTwo(SizeX);
 		SizeY = FMath::RoundUpToPowerOfTwo(SizeY);
 
-		if (PowerOfTwoMode == ETexturePowerOfTwoSetting::PadToSquarePowerOfTwo)
+		if (PowerOfTwoMode == ETexturePowerOfTwoSetting::PadToSquarePowerOfTwo || PowerOfTwoMode == ETexturePowerOfTwoSetting::StretchToSquarePowerOfTwo)
 		{
 			SizeX = SizeY = FMath::Max(SizeX, SizeY);
 		}
diff --git a/Engine/Source/Runtime/Engine/Private/VT/VirtualTextureDataBuilder.cpp b/Engine/Source/Runtime/Engine/Private/VT/VirtualTextureDataBuilder.cpp
--- a/Engine/Source/Runtime/Engine/Private/VT/VirtualTextureDataBuilder.cpp
+++ b/Engine/Source/Runtime/Engine/Private/VT/VirtualTextureDataBuilder.cpp
@@ -261,10 +261,12 @@ bool FVirtualTextureBuilderDerivedInfo::InitializeFromBuildSettings(const FTextu
 	case ETexturePowerOfTwoSetting::None:
 		break;
 	case ETexturePowerOfTwoSetting::PadToPowerOfTwo:
+	case ETexturePowerOfTwoSetting::StretchToPowerOfTwo:
 		BlockSizeX = FMath::RoundUpToPowerOfTwo(BlockSizeX);
 		BlockSizeY = FMath::RoundUpToPowerOfTwo(BlockSizeY);
 		break;
 	case ETexturePowerOfTwoSetting::PadToSquarePowerOfTwo:
+	case ETexturePowerOfTwoSetting::StretchToSquarePowerOfTwo:
 		BlockSizeX = FMath::RoundUpToPowerOfTwo(BlockSizeX);
 		BlockSizeY = FMath::RoundUpToPowerOfTwo(BlockSizeY);
 		BlockSizeX = FMath::Max(BlockSizeX, BlockSizeY);
