Fix for typo in non pow2 volume texture mip count computation.
#10907

diff --git a/Engine/Source/Developer/TextureCompressor/Private/TextureCompressorModule.cpp b/Engine/Source/Developer/TextureCompressor/Private/TextureCompressorModule.cpp
--- a/Engine/Source/Developer/TextureCompressor/Private/TextureCompressorModule.cpp
+++ b/Engine/Source/Developer/TextureCompressor/Private/TextureCompressorModule.cpp
@@ -3491,7 +3491,7 @@ int32 ITextureCompressorModule::GetMipCountForBuildSettings(
 			PowerOfTwoMode != ETexturePowerOfTwoSetting::None)
 		{
 			int32 TargetSizeX, TargetSizeY, TargetSizeZ;
-			bool NeedsAdjustment = GetPowerOfTwoTargetTextureSize(BaseSizeX, BaseSizeY, BaseSizeY, BuildSettings.bVolume, PowerOfTwoMode, TargetSizeX, TargetSizeY, TargetSizeZ);
+			bool NeedsAdjustment = GetPowerOfTwoTargetTextureSize(BaseSizeX, BaseSizeY, BaseSizeZ, BuildSettings.bVolume, PowerOfTwoMode, TargetSizeX, TargetSizeY, TargetSizeZ);
 			if (NeedsAdjustment)
 			{
 				// In this case we are regenerating the entire mip chain.
