Implement PadWithBorderColor texture setting.

diff --git a/Engine/Source/Developer/TextureBuild/Private/TextureBuildFunction.cpp b/Engine/Source/Developer/TextureBuild/Private/TextureBuildFunction.cpp
--- a/Engine/Source/Developer/TextureBuild/Private/TextureBuildFunction.cpp
+++ b/Engine/Source/Developer/TextureBuild/Private/TextureBuildFunction.cpp
@@ -148,6 +148,7 @@ static FTextureBuildSettings ReadBuildSettingsFromCompactBinary(const FCbObjectV
 	BuildSettings.bChromaKeyTexture = Object["bChromaKeyTexture"].AsBool(BuildSettings.bChromaKeyTexture);
 	ReadCbField(Object["PowerOfTwoMode"], BuildSettings.PowerOfTwoMode);
 	ReadCbField(Object["PaddingColor"], BuildSettings.PaddingColor);
+	BuildSettings.bPadWithBorderColor = Object["bPadWithBorderColor"].AsBool(BuildSettings.bPadWithBorderColor);
 	ReadCbField(Object["ChromaKeyColor"], BuildSettings.ChromaKeyColor);
 	ReadCbField(Object["ChromaKeyThreshold"], BuildSettings.ChromaKeyThreshold);
 	ReadCbField(Object["CompressionQuality"], BuildSettings.CompressionQuality);
diff --git a/Engine/Source/Developer/TextureCompressor/Private/TextureCompressorModule.cpp b/Engine/Source/Developer/TextureCompressor/Private/TextureCompressorModule.cpp
--- a/Engine/Source/Developer/TextureCompressor/Private/TextureCompressorModule.cpp
+++ b/Engine/Source/Developer/TextureCompressor/Private/TextureCompressorModule.cpp
@@ -3922,6 +3922,7 @@ private:
 				else
 				{
 					FLinearColor FillColor = BuildSettings.PaddingColor;
+					bool bPadWithBorderColor = BuildSettings.bPadWithBorderColor;
 
 					FLinearColor* TargetPtr = (FLinearColor*)TargetImage.RawData.GetData();
 					FLinearColor* SourcePtr = (FLinearColor*)SourceImage.RawData.GetData();
@@ -3940,10 +3941,18 @@ private:
 								SourcePtr += SourceImage.SizeX;
 								TargetPtr += SourceImage.SizeX;
 							}
+							else if (bPadWithBorderColor)
+							{
+								XStart = TargetImage.SizeX;
+								// We're copying the entirely of the last line of the target image, which we know has proper padding horizontally
+								// because earlier passes (when Y < SourceImage.SizeY) will pad out the line in the loop below.
+								FMemory::Memcpy(TargetPtr, TargetPtr - TargetImage.SizeX, TargetImage.SizeX * sizeof(FLinearColor));
+								TargetPtr += TargetImage.SizeX;
+							}
 
 							for (int32 XPad = XStart; XPad < TargetImage.SizeX; ++XPad)
 							{
-								*TargetPtr++ = FillColor;
+								*TargetPtr++ = bPadWithBorderColor ? *(TargetPtr - 1) : FillColor;
 							}
 						}
 					}
@@ -3952,9 +3961,19 @@ private:
 					{
 						for (int32 Y = 0; Y < TargetImage.SizeY; ++Y)
 						{
-							for (int32 X = 0; X< TargetImage.SizeX; ++X)
+							if (bPadWithBorderColor)
+							{
+								// We're copying the entirely of the corresponding line from the previous slice, which we know has proper padding
+								// performed during earlier passes (SliceIndex < SourceImage.NumSlices).
+								FMemory::Memcpy(TargetPtr, TargetPtr - (int64) TargetImage.SizeX * TargetImage.SizeY, TargetImage.SizeX * sizeof(FLinearColor));
+								TargetPtr += TargetImage.SizeX;
+							}
+							else
 							{
-								*TargetPtr++ = FillColor;
+								for (int32 X = 0; X < TargetImage.SizeX; ++X)
+								{
+									*TargetPtr++ = FillColor;
+								}
 							}
 						}
 					}
diff --git a/Engine/Source/Developer/TextureCompressor/Public/TextureCompressorModule.h b/Engine/Source/Developer/TextureCompressor/Public/TextureCompressorModule.h
--- a/Engine/Source/Developer/TextureCompressor/Public/TextureCompressorModule.h
+++ b/Engine/Source/Developer/TextureCompressor/Public/TextureCompressorModule.h
@@ -185,6 +185,8 @@ struct FTextureBuildSettings
 	uint8 PowerOfTwoMode;
 	/** The color used to pad the texture out if it is resized due to PowerOfTwoMode */
 	FColor PaddingColor;
+	/** If set to true, texture padding will be performed using colors of the border pixels in order to improve quality of the generated mipmaps. */
+	bool bPadWithBorderColor;
 	/** The color that will be replaced with transparent black if chroma keying is enabled */
 	FColor ChromaKeyColor;
 	/** The threshold that components have to match for the texel to be considered equal to the ChromaKeyColor when chroma keying (<=, set to 0 to require a perfect exact match) */
@@ -295,6 +297,7 @@ struct FTextureBuildSettings
 		, bChromaKeyTexture(false)
 		, PowerOfTwoMode(0 /*ETexturePowerOfTwoSetting::None*/)
 		, PaddingColor(FColor::Black)
+		, bPadWithBorderColor(false)
 		, ChromaKeyColor(FColorList::Magenta)
 		, ChromaKeyThreshold(1.0f / 255.0f)
 		, CompressionQuality(-1)
diff --git a/Engine/Source/Runtime/Engine/Classes/Engine/Texture.h b/Engine/Source/Runtime/Engine/Classes/Engine/Texture.h
--- a/Engine/Source/Runtime/Engine/Classes/Engine/Texture.h
+++ b/Engine/Source/Runtime/Engine/Classes/Engine/Texture.h
@@ -1222,10 +1222,14 @@ public:
 	UPROPERTY(EditAnywhere, BlueprintReadWrite, Category=Texture)
 	TEnumAsByte<enum ETexturePowerOfTwoSetting::Type> PowerOfTwoMode;
 
-	/** The color used to pad the texture out if it is resized due to PowerOfTwoMode */
-	UPROPERTY(EditAnywhere, BlueprintReadWrite, Category=Texture)
+	/** The color used to pad the texture out if it is padded due to PowerOfTwoMode */
+	UPROPERTY(EditAnywhere, BlueprintReadWrite, Category=Texture, meta = (EditCondition = "(PowerOfTwoMode == ETexturePowerOfTwoSetting::PadToPowerOfTwo || PowerOfTwoMode == ETexturePowerOfTwoSetting::PadToSquarePowerOfTwo) && !bPadWithBorderColor", EditConditionHides))
 	FColor PaddingColor;
 
+	/** If set to true, texture padding will be performed using colors of the border pixels. This can be used to improve quality of the generated mipmaps for padded textures. */
+	UPROPERTY(EditAnywhere, BlueprintReadWrite, Category=Texture, meta = (EditCondition = "PowerOfTwoMode == ETexturePowerOfTwoSetting::PadToPowerOfTwo || PowerOfTwoMode == ETexturePowerOfTwoSetting::PadToSquarePowerOfTwo", EditConditionHides))
+	bool bPadWithBorderColor;
+
 	/** Whether to chroma key the image, replacing any pixels that match ChromaKeyColor with transparent black */
 	UPROPERTY(EditAnywhere, BlueprintReadWrite, Category=Adjustments)
 	bool bChromaKeyTexture;
diff --git a/Engine/Source/Runtime/Engine/Private/Texture.cpp b/Engine/Source/Runtime/Engine/Private/Texture.cpp
--- a/Engine/Source/Runtime/Engine/Private/Texture.cpp
+++ b/Engine/Source/Runtime/Engine/Private/Texture.cpp
@@ -176,6 +176,7 @@ UTexture::UTexture(const FObjectInitializer& ObjectInitializer)
 	AlphaCoverageThresholds = FVector4(0, 0, 0, 0);
 	bUseNewMipFilter = false;
 	PaddingColor = FColor::Black;
+	bPadWithBorderColor = false;
 	ChromaKeyColor = FColorList::Magenta;
 	ChromaKeyThreshold = 1.0f / 255.0f;
 	VirtualTextureStreaming = 0;
diff --git a/Engine/Source/Runtime/Engine/Private/TextureDerivedData.cpp b/Engine/Source/Runtime/Engine/Private/TextureDerivedData.cpp
--- a/Engine/Source/Runtime/Engine/Private/TextureDerivedData.cpp
+++ b/Engine/Source/Runtime/Engine/Private/TextureDerivedData.cpp
@@ -329,6 +329,12 @@ static void SerializeForKey(FArchive& Ar, const FTextureBuildSettings& Settings)
 		Ar << TempGuid;
 	}
 
+	if (Settings.bPadWithBorderColor)
+	{
+		// @todo SerializeForKey these can go away whenever we bump the overall ddc key
+		TempGuid = FGuid(0xB128BA67, 0x3F3C4797, 0x81C66E55, 0xDEEE78EB);
+		Ar << TempGuid;
+	}
 
 	if ( Settings.bUseNewMipFilter )
 	{
@@ -976,6 +982,7 @@ static void GetTextureBuildSettings(
 	OutBuildSettings.bVirtualStreamable = bVirtualTextureStreaming;
 	OutBuildSettings.PowerOfTwoMode = Texture.PowerOfTwoMode;
 	OutBuildSettings.PaddingColor = Texture.PaddingColor;
+	OutBuildSettings.bPadWithBorderColor = Texture.bPadWithBorderColor;
 	OutBuildSettings.ChromaKeyColor = Texture.ChromaKeyColor;
 	OutBuildSettings.bChromaKeyTexture = Texture.bChromaKeyTexture;
 	OutBuildSettings.ChromaKeyThreshold = Texture.ChromaKeyThreshold;
diff --git a/Engine/Source/Runtime/Engine/Private/TextureDerivedDataBuildUtils.cpp b/Engine/Source/Runtime/Engine/Private/TextureDerivedDataBuildUtils.cpp
--- a/Engine/Source/Runtime/Engine/Private/TextureDerivedDataBuildUtils.cpp
+++ b/Engine/Source/Runtime/Engine/Private/TextureDerivedDataBuildUtils.cpp
@@ -212,6 +212,7 @@ static void WriteBuildSettings(FCbWriter& Writer, const FTextureBuildSettings& B
 	WriteCbFieldWithDefault<bool>(Writer, "bChromaKeyTexture", BuildSettings.bChromaKeyTexture, DefaultSettings.bChromaKeyTexture);
 	WriteCbFieldWithDefault(Writer, "PowerOfTwoMode", BuildSettings.PowerOfTwoMode, DefaultSettings.PowerOfTwoMode);
 	WriteCbFieldWithDefault(Writer, "PaddingColor", BuildSettings.PaddingColor, DefaultSettings.PaddingColor);
+	WriteCbFieldWithDefault<bool>(Writer, "bPadWithBorderColor", BuildSettings.bPadWithBorderColor, DefaultSettings.bPadWithBorderColor);
 	WriteCbFieldWithDefault(Writer, "ChromaKeyColor", BuildSettings.ChromaKeyColor, DefaultSettings.ChromaKeyColor);
 	WriteCbFieldWithDefault(Writer, "ChromaKeyThreshold", BuildSettings.ChromaKeyThreshold, DefaultSettings.ChromaKeyThreshold);
 	WriteCbFieldWithDefault(Writer, "CompressionQuality", BuildSettings.CompressionQuality, DefaultSettings.CompressionQuality);
