Implement Resize to Specific Resolution texture padding and resizing option.

diff --git a/Engine/Source/Developer/TextureBuild/Private/TextureBuildFunction.cpp b/Engine/Source/Developer/TextureBuild/Private/TextureBuildFunction.cpp
--- a/Engine/Source/Developer/TextureBuild/Private/TextureBuildFunction.cpp
+++ b/Engine/Source/Developer/TextureBuild/Private/TextureBuildFunction.cpp
@@ -149,6 +149,8 @@ static FTextureBuildSettings ReadBuildSettingsFromCompactBinary(const FCbObjectV
 	ReadCbField(Object["PowerOfTwoMode"], BuildSettings.PowerOfTwoMode);
 	ReadCbField(Object["PaddingColor"], BuildSettings.PaddingColor);
 	BuildSettings.bPadWithBorderColor = Object["bPadWithBorderColor"].AsBool(BuildSettings.bPadWithBorderColor);
+	ReadCbField(Object["ResizeDuringBuildX"], BuildSettings.ResizeDuringBuildX);
+	ReadCbField(Object["ResizeDuringBuildY"], BuildSettings.ResizeDuringBuildY);
 	ReadCbField(Object["ChromaKeyColor"], BuildSettings.ChromaKeyColor);
 	ReadCbField(Object["ChromaKeyThreshold"], BuildSettings.ChromaKeyThreshold);
 	ReadCbField(Object["CompressionQuality"], BuildSettings.CompressionQuality);
diff --git a/Engine/Source/Developer/TextureCompressor/Private/TextureCompressorModule.cpp b/Engine/Source/Developer/TextureCompressor/Private/TextureCompressorModule.cpp
--- a/Engine/Source/Developer/TextureCompressor/Private/TextureCompressorModule.cpp
+++ b/Engine/Source/Developer/TextureCompressor/Private/TextureCompressorModule.cpp
@@ -3410,7 +3410,7 @@ static void NormalizeMip(FImage& InOutMip)
 
 
 // Returns true if the target texture size is different and padding/stretching is required.
-static bool GetPowerOfTwoTargetTextureSize(int32 InMip0SizeX, int32 InMip0SizeY, int32 InMip0NumSlices, bool bInIsVolume, ETexturePowerOfTwoSetting::Type InPow2Setting, int32& OutTargetSizeX, int32& OutTargetSizeY, int32& OutTargetSizeZ)
+static bool GetPowerOfTwoTargetTextureSize(int32 InMip0SizeX, int32 InMip0SizeY, int32 InMip0NumSlices, bool bInIsVolume, ETexturePowerOfTwoSetting::Type InPow2Setting, int32 InResizeDuringBuildX, int32 InResizeDuringBuildY, int32& OutTargetSizeX, int32& OutTargetSizeY, int32& OutTargetSizeZ)
 {
 	check(InPow2Setting != ETexturePowerOfTwoSetting::None);
 
@@ -3439,6 +3439,17 @@ static bool GetPowerOfTwoTargetTextureSize(int32 InMip0SizeX, int32 InMip0SizeY,
 			FMath::Max3<int32>(PowerOfTwoTextureSizeX, PowerOfTwoTextureSizeY, PowerOfTwoTextureSizeZ);
 		break;
 
+	case ETexturePowerOfTwoSetting::ResizeToSpecificResolution:
+		if (InResizeDuringBuildX)
+		{
+			TargetTextureSizeX = InResizeDuringBuildX;
+		}
+		if (InResizeDuringBuildY)
+		{
+			TargetTextureSizeY = InResizeDuringBuildY;
+		}
+		break;
+
 	default:
 		checkf(false, TEXT("Unknown entry in ETexturePowerOfTwoSetting::Type"));
 		break;
@@ -3491,7 +3502,7 @@ int32 ITextureCompressorModule::GetMipCountForBuildSettings(
 			PowerOfTwoMode != ETexturePowerOfTwoSetting::None)
 		{
 			int32 TargetSizeX, TargetSizeY, TargetSizeZ;
-			bool NeedsAdjustment = GetPowerOfTwoTargetTextureSize(BaseSizeX, BaseSizeY, BaseSizeZ, BuildSettings.bVolume, PowerOfTwoMode, TargetSizeX, TargetSizeY, TargetSizeZ);
+			bool NeedsAdjustment = GetPowerOfTwoTargetTextureSize(BaseSizeX, BaseSizeY, BaseSizeZ, BuildSettings.bVolume, PowerOfTwoMode, BuildSettings.ResizeDuringBuildX, BuildSettings.ResizeDuringBuildY, TargetSizeX, TargetSizeY, TargetSizeZ);
 			if (NeedsAdjustment)
 			{
 				// In this case we are regenerating the entire mip chain.
@@ -3871,23 +3882,23 @@ private:
 			int32 TargetTextureSizeZ = 0;			
 			bool bPadOrStretchTexture = GetPowerOfTwoTargetTextureSize(
 				FirstSourceMipImage.SizeX, FirstSourceMipImage.SizeY, FirstSourceMipImage.NumSlices,
-				BuildSettings.bVolume, PowerOfTwoMode,
+				BuildSettings.bVolume, PowerOfTwoMode, BuildSettings.ResizeDuringBuildX, BuildSettings.ResizeDuringBuildY,
 				TargetTextureSizeX, TargetTextureSizeY, TargetTextureSizeZ);
 
 			if (bPadOrStretchTexture)
 			{
-				bool bStretchTexture = PowerOfTwoMode == ETexturePowerOfTwoSetting::StretchToPowerOfTwo || PowerOfTwoMode == ETexturePowerOfTwoSetting::StretchToSquarePowerOfTwo;
+				bool bResizeTexture = PowerOfTwoMode == ETexturePowerOfTwoSetting::StretchToPowerOfTwo || PowerOfTwoMode == ETexturePowerOfTwoSetting::StretchToSquarePowerOfTwo || PowerOfTwoMode == ETexturePowerOfTwoSetting::ResizeToSpecificResolution;
 				if (BuildSettings.MipGenSettings == TMGS_LeaveExistingMips)
 				{
 					// pad/stretch+leave existing is broken
-					UE_LOG(LogTextureCompressor, Error,	TEXT("Texture padding or stretching to power of two is not allowed when leaving existing mips."));
+					UE_LOG(LogTextureCompressor, Error,	TEXT("Texture padding or resizing is not allowed when leaving existing mips."));
 					return false;
 				}
 				if ( bLongLatCubemap )
 				{
-					if (bStretchTexture)
+					if (bResizeTexture)
 					{
-						UE_LOG(LogTextureCompressor, Warning, TEXT("In order to improve the quality of the generated texture, stretching LongLat cubemaps should be avoided."));
+						UE_LOG(LogTextureCompressor, Warning, TEXT("In order to improve the quality of the generated texture, resizing LongLat cubemaps should be avoided."));
 					}
 					else
 					{
@@ -3909,12 +3920,12 @@ private:
 				const FImage& SourceImage = bSuitableFormat ? FirstSourceMipImage : Temp;
 				FImage& TargetImage = PaddedSourceMips.Emplace_GetRef(TargetTextureSizeX, TargetTextureSizeY, BuildSettings.bVolume ? TargetTextureSizeZ : SourceImage.NumSlices, SourceImage.Format);
 
-				if (bStretchTexture)
+				if (bResizeTexture)
 				{
 					if (TargetImage.NumSlices != 1)
 					{
 						// FImageCore::ResizeTo currently only supports resizing textures with 1 slice
-						UE_LOG(LogTextureCompressor, Error, TEXT("Texture stretching is currently only supported on Texture2D."));
+						UE_LOG(LogTextureCompressor, Error, TEXT("Texture resizing is currently only supported on Texture2D."));
 						return false;
 					}
 					SourceImage.ResizeTo(TargetImage, TargetTextureSizeX, TargetTextureSizeY, SourceImage.Format, SourceImage.GetGammaSpace());
diff --git a/Engine/Source/Developer/TextureCompressor/Public/TextureCompressorModule.h b/Engine/Source/Developer/TextureCompressor/Public/TextureCompressorModule.h
--- a/Engine/Source/Developer/TextureCompressor/Public/TextureCompressorModule.h
+++ b/Engine/Source/Developer/TextureCompressor/Public/TextureCompressorModule.h
@@ -187,6 +187,10 @@ struct FTextureBuildSettings
 	FColor PaddingColor;
 	/** If set to true, texture padding will be performed using colors of the border pixels in order to improve quality of the generated mipmaps. */
 	bool bPadWithBorderColor;
+	/** Width of the resized texture when using "Resize To Specific Resolution" padding and resizing option. If set to zero, original width will be used. */
+	int32 ResizeDuringBuildX;
+	/** Width of the resized texture when using "Resize To Specific Resolution" padding and resizing option. If set to zero, original height will be used. */
+	int32 ResizeDuringBuildY;
 	/** The color that will be replaced with transparent black if chroma keying is enabled */
 	FColor ChromaKeyColor;
 	/** The threshold that components have to match for the texel to be considered equal to the ChromaKeyColor when chroma keying (<=, set to 0 to require a perfect exact match) */
@@ -298,6 +302,8 @@ struct FTextureBuildSettings
 		, PowerOfTwoMode(0 /*ETexturePowerOfTwoSetting::None*/)
 		, PaddingColor(FColor::Black)
 		, bPadWithBorderColor(false)
+		, ResizeDuringBuildX(0)
+		, ResizeDuringBuildY(0)
 		, ChromaKeyColor(FColorList::Magenta)
 		, ChromaKeyThreshold(1.0f / 255.0f)
 		, CompressionQuality(-1)
diff --git a/Engine/Source/Runtime/Engine/Classes/Engine/Texture.h b/Engine/Source/Runtime/Engine/Classes/Engine/Texture.h
--- a/Engine/Source/Runtime/Engine/Classes/Engine/Texture.h
+++ b/Engine/Source/Runtime/Engine/Classes/Engine/Texture.h
@@ -1219,7 +1219,7 @@ public:
 	uint32 bFlipGreenChannel:1;
 
 	/** How to pad the texture to a power of 2 size (if necessary) */
-	UPROPERTY(EditAnywhere, BlueprintReadWrite, Category=Texture)
+	UPROPERTY(EditAnywhere, BlueprintReadWrite, Category=Texture, meta = (DisplayName = "Padding and Resizing"))
 	TEnumAsByte<enum ETexturePowerOfTwoSetting::Type> PowerOfTwoMode;
 
 	/** The color used to pad the texture out if it is padded due to PowerOfTwoMode */
@@ -1230,6 +1230,14 @@ public:
 	UPROPERTY(EditAnywhere, BlueprintReadWrite, Category=Texture, meta = (EditCondition = "PowerOfTwoMode == ETexturePowerOfTwoSetting::PadToPowerOfTwo || PowerOfTwoMode == ETexturePowerOfTwoSetting::PadToSquarePowerOfTwo", EditConditionHides))
 	bool bPadWithBorderColor;
 
+	/** Width of the resized texture when using "Resize To Specific Resolution" padding and resizing option. If set to zero, original width will be used. */
+	UPROPERTY(EditAnywhere, BlueprintReadWrite, Category = Texture, meta = (EditCondition = "PowerOfTwoMode == ETexturePowerOfTwoSetting::ResizeToSpecificResolution", EditConditionHides))
+	int32 ResizeDuringBuildX;
+
+	/** Width of the resized texture when using "Resize To Specific Resolution" padding and resizing option. If set to zero, original height will be used. */
+	UPROPERTY(EditAnywhere, BlueprintReadWrite, Category = Texture, meta = (EditCondition = "PowerOfTwoMode == ETexturePowerOfTwoSetting::ResizeToSpecificResolution", EditConditionHides))
+	int32 ResizeDuringBuildY;
+
 	/** Whether to chroma key the image, replacing any pixels that match ChromaKeyColor with transparent black */
 	UPROPERTY(EditAnywhere, BlueprintReadWrite, Category=Adjustments)
 	bool bChromaKeyTexture;
diff --git a/Engine/Source/Runtime/Engine/Classes/Engine/TextureDefines.h b/Engine/Source/Runtime/Engine/Classes/Engine/TextureDefines.h
--- a/Engine/Source/Runtime/Engine/Classes/Engine/TextureDefines.h
+++ b/Engine/Source/Runtime/Engine/Classes/Engine/TextureDefines.h
@@ -174,7 +174,10 @@ namespace ETexturePowerOfTwoSetting
 		StretchToPowerOfTwo,
 
 		/** Stretch the texture to the nearest square power of two size. */
-		StretchToSquarePowerOfTwo
+		StretchToSquarePowerOfTwo,
+
+		/** Resize the texture to specific user defined resolution. */
+		ResizeToSpecificResolution
 
 		// Note: These are serialized as as raw values in the texture DDC key, so additional entries
 		// should be added at the bottom; reordering or removing entries will require changing the GUID
diff --git a/Engine/Source/Runtime/Engine/Private/Texture.cpp b/Engine/Source/Runtime/Engine/Private/Texture.cpp
--- a/Engine/Source/Runtime/Engine/Private/Texture.cpp
+++ b/Engine/Source/Runtime/Engine/Private/Texture.cpp
@@ -177,6 +177,8 @@ UTexture::UTexture(const FObjectInitializer& ObjectInitializer)
 	bUseNewMipFilter = false;
 	PaddingColor = FColor::Black;
 	bPadWithBorderColor = false;
+	ResizeDuringBuildX = 0;
+	ResizeDuringBuildY = 0;
 	ChromaKeyColor = FColorList::Magenta;
 	ChromaKeyThreshold = 1.0f / 255.0f;
 	VirtualTextureStreaming = 0;
@@ -528,19 +530,22 @@ void UTexture::ValidateSettingsAfterImportOrEdit(bool * pRequiresNotifyMaterials
 		if ( MipGenSettings == TMGS_LeaveExistingMips && PowerOfTwoMode != ETexturePowerOfTwoSetting::None )
 		{
 			// power of 2 padding or stretching is not allowed with LeaveExistingMips
-			UE_LOG(LogTexture, Display, TEXT("Power of 2 padding or stretching cannot be used with LeaveExistingMips, disabled. (%s)"), *GetName());
+			UE_LOG(LogTexture, Display, TEXT("Texture padding or resizing cannot be used with LeaveExistingMips, disabled. (%s)"), *GetName());
 
 			PowerOfTwoMode = ETexturePowerOfTwoSetting::None;
 		}
 
-		if ((PowerOfTwoMode == ETexturePowerOfTwoSetting::StretchToPowerOfTwo || PowerOfTwoMode == ETexturePowerOfTwoSetting::StretchToSquarePowerOfTwo) && !this->IsA<UTexture2D>())
+		if ((PowerOfTwoMode == ETexturePowerOfTwoSetting::StretchToPowerOfTwo || PowerOfTwoMode == ETexturePowerOfTwoSetting::StretchToSquarePowerOfTwo || PowerOfTwoMode == ETexturePowerOfTwoSetting::ResizeToSpecificResolution) && !this->IsA<UTexture2D>())
 		{
-			// currently power of 2 stretching is only supported for 2D textures, but can be implemented for other types of textures in the future
-			UE_LOG(LogTexture, Display, TEXT("Currently power of two stretching is only supported for Texture2D, forcing PowerOfTwoMode to None. (%s)"), *GetName());
+			// currently resizing is only supported for 2D textures, but can be implemented for other types of textures in the future
+			UE_LOG(LogTexture, Display, TEXT("Currently resizing is only supported for Texture2D, forcing PowerOfTwoMode to None. (%s)"), *GetName());
 
 			PowerOfTwoMode = ETexturePowerOfTwoSetting::None;
 		}
 
+		ResizeDuringBuildX = FMath::Max(0, FMath::Min((int32)GetMaximumDimension(), ResizeDuringBuildX));
+		ResizeDuringBuildY = FMath::Max(0, FMath::Min((int32)GetMaximumDimension(), ResizeDuringBuildY));
+
 		// IsPowerOfTwo only checks XY :
 		bool bIsPowerOfTwo = Source.IsPowerOfTwo();
 		if ( ! FMath::IsPowerOfTwo(Source.GetVolumeSizeZ()) )
@@ -3346,6 +3351,17 @@ void UTexture::GetBuiltTextureSize( const ITargetPlatform* TargetPlatform , int3
 			SizeX = SizeY = FMath::Max(SizeX, SizeY);
 		}
 	}
+	else if (PowerOfTwoMode == ETexturePowerOfTwoSetting::ResizeToSpecificResolution)
+	{
+		if (ResizeDuringBuildX)
+		{
+			SizeX = ResizeDuringBuildX;
+		}
+		if (ResizeDuringBuildY)
+		{
+			SizeY = ResizeDuringBuildY;
+		}
+	}
 	else
 	{
 		checkf(PowerOfTwoMode == ETexturePowerOfTwoSetting::None, TEXT("Unknown entry in ETexturePowerOfTwoSetting::Type"));
diff --git a/Engine/Source/Runtime/Engine/Private/TextureDerivedData.cpp b/Engine/Source/Runtime/Engine/Private/TextureDerivedData.cpp
--- a/Engine/Source/Runtime/Engine/Private/TextureDerivedData.cpp
+++ b/Engine/Source/Runtime/Engine/Private/TextureDerivedData.cpp
@@ -336,6 +336,14 @@ static void SerializeForKey(FArchive& Ar, const FTextureBuildSettings& Settings)
 		Ar << TempGuid;
 	}
 
+	if (Settings.ResizeDuringBuildX || Settings.ResizeDuringBuildY)
+	{
+		// @todo SerializeForKey these can go away whenever we bump the overall ddc key
+		TempGuid = FGuid(0xDAE8B3E9, 0x605B49DC, 0xADA3C221, 0x02D5567D); Ar << TempGuid;
+		TempUint32 = Settings.ResizeDuringBuildX; Ar << TempUint32;
+		TempUint32 = Settings.ResizeDuringBuildY; Ar << TempUint32;
+	}
+
 	if ( Settings.bUseNewMipFilter )
 	{
 		// @todo SerializeForKey : TextureAddressModeX is only used if bUseNewMipFilter is true
@@ -983,6 +991,8 @@ static void GetTextureBuildSettings(
 	OutBuildSettings.PowerOfTwoMode = Texture.PowerOfTwoMode;
 	OutBuildSettings.PaddingColor = Texture.PaddingColor;
 	OutBuildSettings.bPadWithBorderColor = Texture.bPadWithBorderColor;
+	OutBuildSettings.ResizeDuringBuildX = Texture.ResizeDuringBuildX;
+	OutBuildSettings.ResizeDuringBuildY = Texture.ResizeDuringBuildY;
 	OutBuildSettings.ChromaKeyColor = Texture.ChromaKeyColor;
 	OutBuildSettings.bChromaKeyTexture = Texture.bChromaKeyTexture;
 	OutBuildSettings.ChromaKeyThreshold = Texture.ChromaKeyThreshold;
diff --git a/Engine/Source/Runtime/Engine/Private/TextureDerivedDataBuildUtils.cpp b/Engine/Source/Runtime/Engine/Private/TextureDerivedDataBuildUtils.cpp
--- a/Engine/Source/Runtime/Engine/Private/TextureDerivedDataBuildUtils.cpp
+++ b/Engine/Source/Runtime/Engine/Private/TextureDerivedDataBuildUtils.cpp
@@ -213,6 +213,8 @@ static void WriteBuildSettings(FCbWriter& Writer, const FTextureBuildSettings& B
 	WriteCbFieldWithDefault(Writer, "PowerOfTwoMode", BuildSettings.PowerOfTwoMode, DefaultSettings.PowerOfTwoMode);
 	WriteCbFieldWithDefault(Writer, "PaddingColor", BuildSettings.PaddingColor, DefaultSettings.PaddingColor);
 	WriteCbFieldWithDefault<bool>(Writer, "bPadWithBorderColor", BuildSettings.bPadWithBorderColor, DefaultSettings.bPadWithBorderColor);
+	WriteCbFieldWithDefault(Writer, "ResizeDuringBuildX", BuildSettings.ResizeDuringBuildX, DefaultSettings.ResizeDuringBuildX);
+	WriteCbFieldWithDefault(Writer, "ResizeDuringBuildY", BuildSettings.ResizeDuringBuildY, DefaultSettings.ResizeDuringBuildY);
 	WriteCbFieldWithDefault(Writer, "ChromaKeyColor", BuildSettings.ChromaKeyColor, DefaultSettings.ChromaKeyColor);
 	WriteCbFieldWithDefault(Writer, "ChromaKeyThreshold", BuildSettings.ChromaKeyThreshold, DefaultSettings.ChromaKeyThreshold);
 	WriteCbFieldWithDefault(Writer, "CompressionQuality", BuildSettings.CompressionQuality, DefaultSettings.CompressionQuality);
diff --git a/Engine/Source/Runtime/Engine/Private/VT/VirtualTextureDataBuilder.cpp b/Engine/Source/Runtime/Engine/Private/VT/VirtualTextureDataBuilder.cpp
--- a/Engine/Source/Runtime/Engine/Private/VT/VirtualTextureDataBuilder.cpp
+++ b/Engine/Source/Runtime/Engine/Private/VT/VirtualTextureDataBuilder.cpp
@@ -272,6 +272,16 @@ bool FVirtualTextureBuilderDerivedInfo::InitializeFromBuildSettings(const FTextu
 		BlockSizeX = FMath::Max(BlockSizeX, BlockSizeY);
 		BlockSizeY = BlockSizeX;
 		break;
+	case ETexturePowerOfTwoSetting::ResizeToSpecificResolution:
+		if (BuildSettingsLayer0.ResizeDuringBuildX)
+		{
+			BlockSizeX = BuildSettingsLayer0.ResizeDuringBuildX;
+		}
+		if (BuildSettingsLayer0.ResizeDuringBuildY)
+		{
+			BlockSizeY = BuildSettingsLayer0.ResizeDuringBuildY;
+		}
+		break;
 	default:
 		checkNoEntry();
 		break;
