nop comments only; mark some of the bad image resizers in UE and redirect users to the new ResizeImage
not actually deprecated yet, need to keep legacy behavior to avoid patches
todo actually remove use of the old resize functions and deprecate them

#rnx

diff --git a/Engine/Plugins/Enterprise/MDLImporter/Source/MDLImporter/Private/common/ImageUtils.cpp b/Engine/Plugins/Enterprise/MDLImporter/Source/MDLImporter/Private/common/ImageUtils.cpp
--- a/Engine/Plugins/Enterprise/MDLImporter/Source/MDLImporter/Private/common/ImageUtils.cpp
+++ b/Engine/Plugins/Enterprise/MDLImporter/Source/MDLImporter/Private/common/ImageUtils.cpp
@@ -16,6 +16,8 @@ DEFINE_LOG_CATEGORY_STATIC(LogImageUtils, Log, All);
 namespace ImageUtilsImpl
 {
 
+	// TODO: get rid of all this, use FImage instead.
+
 	bool GetRawData(UTextureRenderTarget2D* TexRT, TArray<uint8>& RawData)
 	{
 		FRenderTarget* RenderTarget = TexRT->GameThread_GetRenderTargetResource();
@@ -96,6 +98,8 @@ namespace ImageUtilsImpl
 	{
 		check(SrcChannels > 0 && SrcChannels <= 4);
 
+		// TODO: do not use this, use FImageCore::ResizeImage instead
+
 		for (int32 Index = 0; Index < DstWidth * DstHeight * SrcChannels; ++Index)
 		{
 			DstData[Index] = (T)0;
@@ -180,6 +184,8 @@ namespace Common{
 /**
  * Resizes the given image using a simple average filter and stores it in the destination array.
  *
+ * TODO: do not use this, use FImageCore::ResizeImage instead
+ *
  * @param SrcWidth		Source image width.
  * @param SrcHeight		Source image height.
  * @param SrcData		Source image data.
diff --git a/Engine/Source/Developer/TextureCompressor/Private/TextureCompressorModule.cpp b/Engine/Source/Developer/TextureCompressor/Private/TextureCompressorModule.cpp
--- a/Engine/Source/Developer/TextureCompressor/Private/TextureCompressorModule.cpp
+++ b/Engine/Source/Developer/TextureCompressor/Private/TextureCompressorModule.cpp
@@ -1469,7 +1469,11 @@ static void DownscaleImage(const FImage& SrcImage, FImage& DstImage, const FText
 	float Downscale = GetDownscaleFinalSizeAndClampedDownscale(SrcImage.SizeX, SrcImage.SizeY, Settings, FinalSizeX, FinalSizeY);
 	
 	//@todo OodleImageResize : replace this whole function with better image resizer if NewFilters
-
+	//	use FImageCore::ResizeImage
+	// what this function does is 2X downsamples with the mip filter
+	//	and then a final bilinear resize to the desired final size
+	// instead just use ResizeImage to go directly from source size to final size in one step
+	
 	// recompute Downscale factor because it may have changed due to block alignment
 	// note: if aspect ratio was not exactly preserved, this could differ in X and Y
 	Downscale = (float)SrcImage.SizeX / (float)FinalSizeX;
@@ -1560,6 +1564,7 @@ static void DownscaleImage(const FImage& SrcImage, FImage& DstImage, const FText
 	FImageView2D DstImageData(*ImageChain[1], 0);
 				
 	// @todo OodleImageResize : not sure this is a correct image resize without shift; does it get pixel center offsets right?
+	//	use FImageCore::ResizeImage here instead
 	for (int32 Y = 0; Y < FinalSizeY; ++Y)
 	{
 		float SourceY = (float)Y * Downscale;
@@ -1568,12 +1573,12 @@ static void DownscaleImage(const FImage& SrcImage, FImage& DstImage, const FText
 		for (int32 X = 0; X < FinalSizeX; ++X)
 		{
 			float SourceX = (float)X * Downscale;
-			int32 IntSourceX = FMath::RoundToInt(SourceX);
 
 			FLinearColor FilteredColor(0,0,0,0);
 
 			if (bUnfiltered)
 			{
+				int32 IntSourceX = FMath::RoundToInt(SourceX);
 				FilteredColor = LookupSourceMip<MGTAM_Clamp>(SrcImageData, IntSourceX, IntSourceY);
 			}
 			else if(bBilinear)
@@ -3924,10 +3929,13 @@ private:
 				{
 					if (TargetImage.NumSlices != 1)
 					{
+						// @todo OodleImageResize ; use FImageCore::ResizeImage instead, supports slices if slice count is the same (eg. arrays yes, volumes no)
 						// FImageCore::ResizeTo currently only supports resizing textures with 1 slice
 						UE_LOG(LogTextureCompressor, Error, TEXT("Texture resizing is currently only supported on Texture2D."));
 						return false;
 					}
+					
+					// @todo OodleImageResize : use FImageCore::ResizeImage instead
 					SourceImage.ResizeTo(TargetImage, TargetTextureSizeX, TargetTextureSizeY, SourceImage.Format, SourceImage.GetGammaSpace());
 				}
 				else
diff --git a/Engine/Source/Runtime/Engine/Private/ImageUtils.cpp b/Engine/Source/Runtime/Engine/Private/ImageUtils.cpp
--- a/Engine/Source/Runtime/Engine/Private/ImageUtils.cpp
+++ b/Engine/Source/Runtime/Engine/Private/ImageUtils.cpp
@@ -408,6 +408,8 @@ bool FImageUtils::GetRenderTargetImage(UTextureRenderTarget2D* TexRT, FImage & I
  */
 void FImageUtils::ImageResize(int32 SrcWidth, int32 SrcHeight, const TArray<FColor> &SrcData, int32 DstWidth, int32 DstHeight, TArray<FColor> &DstData, bool bLinearSpace, bool bForceOpaqueOutput)
 {
+	// * DEPRECATED do not use this, use FImageCore::ResizeImage instead
+
 	DstData.Empty(DstWidth*DstHeight);
 	DstData.AddZeroed(DstWidth*DstHeight);
 
@@ -427,7 +429,8 @@ void FImageUtils::ImageResize(int32 SrcWidth, int32 SrcHeight, const TArray<FCol
  */
 void FImageUtils::ImageResize(int32 SrcWidth, int32 SrcHeight, const TArrayView<const FColor> &SrcData, int32 DstWidth, int32 DstHeight, const TArrayView<FColor> &DstData, bool bLinearSpace, bool bForceOpaqueOutput)
 {
-	//@todo OodleImageResize : deprecate ImageResize and direct users to new function
+	//@todo OodleImageResize : deprecate ImageResize
+	// * DEPRECATED do not use this, use FImageCore::ResizeImage instead
 
 	check(SrcData.Num() >= SrcWidth * SrcHeight);
 	check(DstData.Num() >= DstWidth * DstHeight);
@@ -532,6 +535,8 @@ void FImageUtils::ImageResize(int32 SrcWidth, int32 SrcHeight, const TArrayView<
  */
 void FImageUtils::ImageResize(int32 SrcWidth, int32 SrcHeight, const TArray64<FLinearColor>& SrcData, int32 DstWidth, int32 DstHeight, TArray64<FLinearColor>& DstData)
 {
+	// * DEPRECATED do not use this, use FImageCore::ResizeImage instead
+
 	DstData.Empty(DstWidth * DstHeight);
 	DstData.AddZeroed(DstWidth * DstHeight);
 
@@ -551,7 +556,8 @@ void FImageUtils::ImageResize(int32 SrcWidth, int32 SrcHeight, const TArray64<FL
  */
 void FImageUtils::ImageResize(int32 SrcWidth, int32 SrcHeight, const TArrayView64<const FLinearColor>& SrcData, int32 DstWidth, int32 DstHeight, const TArrayView64<FLinearColor>& DstData)
 {
-	//@todo OodleImageResize : deprecate ImageResize and direct users to new function
+	//@todo OodleImageResize : deprecate ImageResize
+	// * DEPRECATED do not use this, use FImageCore::ResizeImage instead
 
 	check(SrcData.Num() >= SrcWidth * SrcHeight);
 	check(DstData.Num() >= DstWidth * DstHeight);
diff --git a/Engine/Source/Runtime/Engine/Public/ImageUtils.h b/Engine/Source/Runtime/Engine/Public/ImageUtils.h
--- a/Engine/Source/Runtime/Engine/Public/ImageUtils.h
+++ b/Engine/Source/Runtime/Engine/Public/ImageUtils.h
@@ -161,6 +161,8 @@ public:
 	/**
 	 * Resizes the given image using a simple average filter and stores it in the destination array.  This version constrains aspect ratio.
 	 *
+	 * DEPRECATED do not use this, use FImageCore::ResizeImage instead
+	 *
 	 * @param SrcWidth				Source image width.
 	 * @param SrcHeight				Source image height.
 	 * @param SrcData				Source image data.
@@ -176,6 +178,8 @@ public:
 	 * Resizes the given image using a simple average filter and stores it in the destination array.  This version constrains aspect ratio.
 	 * Accepts TArrayViews but requires that DstData be pre-sized appropriately
 	 *
+	 * DEPRECATED do not use this, use FImageCore::ResizeImage instead
+	 *
 	 * @param SrcWidth				Source image width.
 	 * @param SrcHeight				Source image height.
 	 * @param SrcData				Source image data.
@@ -190,6 +194,8 @@ public:
 	/**
 	 * Resizes the given image using a simple average filter and stores it in the destination array.  This version constrains aspect ratio.
 	 *
+	 * DEPRECATED do not use this, use FImageCore::ResizeImage instead
+	 *
 	 * @param SrcWidth	Source image width.
 	 * @param SrcHeight	Source image height.
 	 * @param SrcData	Source image data.
@@ -203,6 +209,8 @@ public:
 	 * Resizes the given image using a simple average filter and stores it in the destination array.  This version constrains aspect ratio.
 	 * Accepts TArrayViews but requires that DstData be pre-sized appropriately
 	 *
+	 * DEPRECATED do not use this, use FImageCore::ResizeImage instead
+	 *
 	 * @param SrcWidth	Source image width.
 	 * @param SrcHeight	Source image height.
 	 * @param SrcData	Source image data.
