Texture: make stretch-to-pow2 use new ResizeImage
bump ddc key for affected textures

#rnx

diff --git a/Engine/Source/Developer/TextureCompressor/Private/TextureCompressorModule.cpp b/Engine/Source/Developer/TextureCompressor/Private/TextureCompressorModule.cpp
--- a/Engine/Source/Developer/TextureCompressor/Private/TextureCompressorModule.cpp
+++ b/Engine/Source/Developer/TextureCompressor/Private/TextureCompressorModule.cpp
@@ -3926,17 +3926,12 @@ private:
 				FImage& TargetImage = PaddedSourceMips.Emplace_GetRef(TargetTextureSizeX, TargetTextureSizeY, BuildSettings.bVolume ? TargetTextureSizeZ : SourceImage.NumSlices, SourceImage.Format);
 
 				if (bResizeTexture)
-				{
-					if (TargetImage.NumSlices != 1)
-					{
-						// @todo OodleImageResize ; use FImageCore::ResizeImage instead, supports slices if slice count is the same (eg. arrays yes, volumes no)
-						// FImageCore::ResizeTo currently only supports resizing textures with 1 slice
-						UE_LOG(LogTextureCompressor, Error, TEXT("Texture resizing is currently only supported on Texture2D."));
-						return false;
-					}
-					
-					// @todo OodleImageResize : use FImageCore::ResizeImage instead
-					SourceImage.ResizeTo(TargetImage, TargetTextureSizeX, TargetTextureSizeY, SourceImage.Format, SourceImage.GetGammaSpace());
+				{		
+					// changed resizer for "stretch to power of two" , bumped ddc key
+
+					// choice of Filter?
+
+					FImageCore::ResizeImageAllocDest(SourceImage,TargetImage, TargetTextureSizeX, TargetTextureSizeY);
 				}
 				else
 				{
@@ -3952,26 +3947,44 @@ private:
 					{
 						for (int32 Y = 0; Y < TargetTextureSizeY; ++Y)
 						{
-							int32 XStart = 0;
 							if (Y < SourceImage.SizeY)
 							{
-								XStart = SourceImage.SizeX;
 								FMemory::Memcpy(TargetPtr, SourcePtr, SourceImage.SizeX * sizeof(FLinearColor));
 								SourcePtr += SourceImage.SizeX;
 								TargetPtr += SourceImage.SizeX;
+								
+								int32 PadCount = TargetImage.SizeX - SourceImage.SizeX;
+								if ( bPadWithBorderColor )
+								{
+									const FLinearColor BorderColor = TargetPtr[-1];
+									for (int32 i=0;i<PadCount;i++)
+									{
+										*TargetPtr++ = BorderColor;
+									}
+								}
+								else
+								{
+									for (int32 i=0;i<PadCount;i++)
+									{
+										*TargetPtr++ = FillColor;
+									}
+								}
 							}
 							else if (bPadWithBorderColor)
 							{
-								XStart = TargetImage.SizeX;
 								// We're copying the entirely of the last line of the target image, which we know has proper padding horizontally
 								// because earlier passes (when Y < SourceImage.SizeY) will pad out the line in the loop below.
 								FMemory::Memcpy(TargetPtr, TargetPtr - TargetImage.SizeX, TargetImage.SizeX * sizeof(FLinearColor));
 								TargetPtr += TargetImage.SizeX;
 							}
-
-							for (int32 XPad = XStart; XPad < TargetImage.SizeX; ++XPad)
+							else
 							{
-								*TargetPtr++ = bPadWithBorderColor ? *(TargetPtr - 1) : FillColor;
+								// fill whole line with FillColor
+								
+								for (int32 i=0;i<TargetImage.SizeX;i++)
+								{
+									*TargetPtr++ = FillColor;
+								}
 							}
 						}
 					}
diff --git a/Engine/Source/Runtime/Engine/Private/TextureDerivedData.cpp b/Engine/Source/Runtime/Engine/Private/TextureDerivedData.cpp
--- a/Engine/Source/Runtime/Engine/Private/TextureDerivedData.cpp
+++ b/Engine/Source/Runtime/Engine/Private/TextureDerivedData.cpp
@@ -215,6 +215,14 @@ static void SerializeForKey(FArchive& Ar, const FTextureBuildSettings& Settings)
 	TempColor = Settings.ChromaKeyColor; Ar << TempColor;
 	TempFloat = Settings.ChromaKeyThreshold; Ar << TempFloat;
 	
+	if ( Settings.PowerOfTwoMode >= ETexturePowerOfTwoSetting::Type::StretchToPowerOfTwo )
+	{
+		// @todo SerializeForKey these can go away whenever we bump the overall ddc key
+		// Stretch power of two modes ResizeImage changed 10-31-2023
+		TempGuid = FGuid(0xb88aa846, 0xadec4199, 0x9a3cf2f2, 0x1413abc6);
+		Ar << TempGuid;
+	}
+
 	// Avoid changing key for non-VT enabled textures
 	if (Settings.bVirtualStreamable)
 	{
