Texture Downscale use new ResizeImage, if NewFilter

#rnx

diff --git a/Engine/Source/Developer/TextureCompressor/Private/TextureCompressorModule.cpp b/Engine/Source/Developer/TextureCompressor/Private/TextureCompressorModule.cpp
--- a/Engine/Source/Developer/TextureCompressor/Private/TextureCompressorModule.cpp
+++ b/Engine/Source/Developer/TextureCompressor/Private/TextureCompressorModule.cpp
@@ -1403,48 +1403,67 @@ static float GetDownscaleFinalSizeAndClampedDownscale(int32 SrcImageWidth, int32
 	check(Settings.Downscale > 1.0f); // must be already handled.
 	
 	float Downscale = FMath::Clamp(Settings.Downscale, 1.f, 8.f);
+	
+	// currently BlockSize always == 4
+	// if source was blocksize aligned, make sure final size is too
+	//	this is required if pixel format is DXT
+	bool bKeepBlockSizeAligned = (Settings.BlockSize > 1
+			&& (SrcImageWidth  % Settings.BlockSize) == 0
+			&& (SrcImageHeight % Settings.BlockSize) == 0);
 
-	// note: more accurate would be to use FMath::Max(1, FMath::RoundToInt(SrcImage.SizeX / Downscale))
-	int32 FinalSizeX = FMath::CeilToInt((float)SrcImageWidth / Downscale);
-	int32 FinalSizeY = FMath::CeilToInt((float)SrcImageHeight / Downscale);
-
-	// compute final size respecting image block size
-	if (Settings.BlockSize > 1
-		&& (SrcImageWidth % Settings.BlockSize) == 0
-		&& (SrcImageHeight % Settings.BlockSize) == 0)
-	{
-		// the following code finds non-zero dimensions of the scaled image which preserve both aspect ratio and block alignment, 
-		// it favors preserving aspect ratio at the expense of not scaling the desired factor when both are not possible
-		int32 GCD = FMath::GreatestCommonDivisor(SrcImageWidth, SrcImageHeight);
-		int32 ScalingGridSizeX = (SrcImageWidth / GCD) * Settings.BlockSize;
-		// note: more accurate would be to use (SrcImage.SizeX / Downscale) instead of FinalSizeX here
-		// GridSnap rounds to nearest, and can return zero
-		FinalSizeX = FMath::GridSnap(FinalSizeX, ScalingGridSizeX);
-		FinalSizeX = FMath::Max(ScalingGridSizeX, FinalSizeX);
-		FinalSizeY = (int32)( ((int64)FinalSizeX * SrcImageHeight) / SrcImageWidth );
-		// Final Size X and Y are gauranteed to be block aligned
-
-		#if 0
-		
-		// simpler alternative :
-		// choose the block count in the smaller dimension first
-		// then make the larger dimension maintain aspect ratio
-		int32 FinalNumBlocksX,FinalNumBlocksY;
-		if ( SrcImage.SizeX >= SrcImage.SizeY )
+	int32 FinalSizeX,FinalSizeY;
+
+	if ( Settings.UseNewMipFilter )
+	{
+		// new way:
+
+		if ( bKeepBlockSizeAligned )
 		{
-			FinalNumBlocksY = FMath::RoundToInt( SrcImage.SizeY / (Downscale * Settings.BlockSize) );
-			FinalNumBlocksX = FMath::RoundToInt( FinalNumBlocksY * SrcImage.SizeX / (float)SrcImage.SizeY );
+			// choose the block count in the smaller dimension first
+			// then make the larger dimension maintain aspect ratio
+			//  we favor block alignment and getting the desired downscale over exactly preserving aspect ratio
+			int32 FinalNumBlocksX,FinalNumBlocksY;
+			if ( SrcImageWidth >= SrcImageHeight )
+			{
+				FinalNumBlocksY = FMath::RoundToInt( SrcImageHeight / (Downscale * Settings.BlockSize) );
+				FinalNumBlocksX = FMath::RoundToInt( FinalNumBlocksY * SrcImageWidth / (float)SrcImageHeight );
+			}
+			else
+			{
+				FinalNumBlocksX = FMath::RoundToInt( SrcImageWidth / (Downscale * Settings.BlockSize) );
+				FinalNumBlocksY = FMath::RoundToInt( FinalNumBlocksX * SrcImageHeight / (float)SrcImageWidth );
+			}
+
+			FinalSizeX = FMath::Max(FinalNumBlocksX,1)*Settings.BlockSize;
+			FinalSizeY = FMath::Max(FinalNumBlocksY,1)*Settings.BlockSize;
 		}
 		else
 		{
-			FinalNumBlocksX = FMath::RoundToInt( SrcImage.SizeX / (Downscale * Settings.BlockSize) );
-			FinalNumBlocksY = FMath::RoundToInt( FinalNumBlocksX * SrcImage.SizeY / (float)SrcImage.SizeX );
+			FinalSizeX = FMath::Max(1, FMath::RoundToInt(SrcImageWidth / Downscale));
+			FinalSizeY = FMath::Max(1, FMath::RoundToInt(SrcImageHeight / Downscale));
 		}
+	}
+	else
+	{
+		// old way :
+		//  this has the flaw of being very far away from the desired downscale when Width/Height have no good GCD
 
-		FinalSizeX = FMath::Max(FinalNumBlocksX,1)*Settings.BlockSize;
-		FinalSizeY = FMath::Max(FinalNumBlocksY,1)*Settings.BlockSize;
+		FinalSizeX = FMath::CeilToInt((float)SrcImageWidth / Downscale);
+		FinalSizeY = FMath::CeilToInt((float)SrcImageHeight / Downscale);
 
-		#endif
+		if ( bKeepBlockSizeAligned )
+		{
+			// the following code finds non-zero dimensions of the scaled image which preserve both aspect ratio and block alignment, 
+			// it favors preserving aspect ratio at the expense of not scaling the desired factor when both are not possible
+			int32 GCD = FMath::GreatestCommonDivisor(SrcImageWidth, SrcImageHeight);
+			int32 ScalingGridSizeX = (SrcImageWidth / GCD) * Settings.BlockSize;
+			// note: more accurate would be to use (SrcImage.SizeX / Downscale) instead of FinalSizeX here
+			// GridSnap rounds to nearest, and can return zero
+			FinalSizeX = FMath::GridSnap(FinalSizeX, ScalingGridSizeX);
+			FinalSizeX = FMath::Max(ScalingGridSizeX, FinalSizeX);
+			FinalSizeY = (int32)( ((int64)FinalSizeX * SrcImageHeight) / SrcImageWidth );
+			// Final Size X and Y are gauranteed to be block aligned
+		}
 	}
 
 	OutWidth = FinalSizeX;
@@ -1454,6 +1473,7 @@ static float GetDownscaleFinalSizeAndClampedDownscale(int32 SrcImageWidth, int32
 
 static void DownscaleImage(const FImage& SrcImage, FImage& DstImage, const FTextureDownscaleSettings& Settings)
 {
+	// BEWARE: DownscaleImage is called with SrcImage == DstImage for in-place operation
 	if (Settings.Downscale <= 1.f)
 	{
 		return;
@@ -1467,9 +1487,31 @@ static void DownscaleImage(const FImage& SrcImage, FImage& DstImage, const FText
 
 	int32 FinalSizeX = 0, FinalSizeY =0;
 	float Downscale = GetDownscaleFinalSizeAndClampedDownscale(SrcImage.SizeX, SrcImage.SizeY, Settings, FinalSizeX, FinalSizeY);
+
+	if ( Settings.UseNewMipFilter )
+	{
+		// changed DDC key for affected textures
+				
+		// choose Filter from ETextureDownscaleOptions
+		FImageCore::EResizeImageFilter Filter;
+
+		if ( Settings.DownscaleOptions == (uint8)ETextureDownscaleOptions::Unfiltered )
+			Filter = FImageCore::EResizeImageFilter::PointSample;
+		else if ( Settings.DownscaleOptions == (uint8)ETextureDownscaleOptions::SimpleAverage )
+			Filter = FImageCore::EResizeImageFilter::Triangle;
+		else if ( Settings.DownscaleOptions >= (uint8)ETextureDownscaleOptions::Sharpen5 ) // sharpest
+			Filter = FImageCore::EResizeImageFilter::CubicSharp;
+		else if ( Settings.DownscaleOptions >= (uint8)ETextureDownscaleOptions::Sharpen0 ) // sharp
+			Filter = FImageCore::EResizeImageFilter::AdaptiveSharp;		
+		else
+			Filter = FImageCore::EResizeImageFilter::AdaptiveSmooth;		
+
+		FImage Temp; // needed because Src == Dst
+		FImageCore::ResizeImageAllocDest(SrcImage,Temp,FinalSizeX,FinalSizeY,Filter);
+		Temp.Swap(DstImage);
+		return;
+	}
 	
-	//@todo OodleImageResize : replace this whole function with better image resizer if NewFilters
-	//	use FImageCore::ResizeImage
 	// what this function does is 2X downsamples with the mip filter
 	//	and then a final bilinear resize to the desired final size
 	// instead just use ResizeImage to go directly from source size to final size in one step
diff --git a/Engine/Source/Runtime/Engine/Private/TextureDerivedData.cpp b/Engine/Source/Runtime/Engine/Private/TextureDerivedData.cpp
--- a/Engine/Source/Runtime/Engine/Private/TextureDerivedData.cpp
+++ b/Engine/Source/Runtime/Engine/Private/TextureDerivedData.cpp
@@ -252,6 +252,14 @@ static void SerializeForKey(FArchive& Ar, const FTextureBuildSettings& Settings)
 	{
 		TempFloat = Settings.Downscale; Ar << TempFloat;
 		TempByte = Settings.DownscaleOptions; Ar << TempByte;
+
+		if ( Settings.bUseNewMipFilter )
+		{
+			// downscale behavior changed
+			// @todo SerializeForKey these can go away whenever we bump the overall ddc key
+			TempGuid = FGuid(0xBC9D413B, 0x2C9DF1E3, 0xBF963C7A, 0xABADF00D);
+			Ar << TempGuid;
+		}
 	}
 
 	// this is done in a funny way to add the bool that wasn't being serialized before
