TextureCompressorModule : fast path for common SourceEncoding conversions
sRGB conversion is 2X faster just because it is one pass over the image instead of two

#rnx

diff --git a/Engine/Source/Developer/TextureCompressor/Private/TextureCompressorModule.cpp b/Engine/Source/Developer/TextureCompressor/Private/TextureCompressorModule.cpp
--- a/Engine/Source/Developer/TextureCompressor/Private/TextureCompressorModule.cpp
+++ b/Engine/Source/Developer/TextureCompressor/Private/TextureCompressorModule.cpp
@@ -1657,6 +1657,31 @@ static void LinearizeToWorkingColorSpace(const FImage& SrcImage, FImage& DstImag
 
 	FImageView SrcImageView(SrcImage);
 	const UE::Color::EEncoding SourceEncodingOverride = static_cast<UE::Color::EEncoding>(BuildSettings.SourceEncodingOverride);
+	
+	if ( ! BuildSettings.bHasColorSpaceDefinition )
+	{
+		if ( SourceEncodingOverride == UE::Color::EEncoding::Linear )
+		{
+			SrcImageView.GammaSpace = EGammaSpace::Linear;
+			FImageCore::CopyImage(SrcImageView, DstImage);
+			return;
+		}
+		else if ( SourceEncodingOverride == UE::Color::EEncoding::sRGB 
+			&& ERawImageFormat::GetFormatNeedsGammaSpace(SrcImage.Format) )
+		{
+			SrcImageView.GammaSpace = EGammaSpace::sRGB;
+			FImageCore::CopyImage(SrcImageView, DstImage);
+			return;
+		}
+		else if ( SourceEncodingOverride == UE::Color::EEncoding::Gamma22 
+			&& ERawImageFormat::GetFormatNeedsGammaSpace(SrcImage.Format) )
+		{
+			SrcImageView.GammaSpace = EGammaSpace::Pow22;
+			FImageCore::CopyImage(SrcImageView, DstImage);
+			return;
+		}
+		// could also early out for Encoding == None, but that's handled below
+	}
 
 	// If the source encoding override is active, we avoid CopyImage's de-gammatization and instead let OpenColorIO do the decoding below.
 	if (SourceEncodingOverride != UE::Color::EEncoding::None)
@@ -1664,6 +1689,7 @@ static void LinearizeToWorkingColorSpace(const FImage& SrcImage, FImage& DstImag
 		SrcImageView.GammaSpace = DstImage.GammaSpace; // EGammaSpace::Linear
 	}
 
+        // CopyImage to get pixels in RGAB32F , then OCIO will act on those in-place in DstImage
 	FImageCore::CopyImage(SrcImageView, DstImage);
 
 	// Decode and/or color transform to the working color space when needed
