Testing the alpha detection based on the channel source info.

diff --git a/Engine/Source/Developer/TextureCompressor/Private/TextureCompressorModule.cpp b/Engine/Source/Developer/TextureCompressor/Private/TextureCompressorModule.cpp
--- a/Engine/Source/Developer/TextureCompressor/Private/TextureCompressorModule.cpp
+++ b/Engine/Source/Developer/TextureCompressor/Private/TextureCompressorModule.cpp
@@ -26,6 +26,13 @@
 #include "Windows/WindowsHWrapper.h"
 #endif
 
+static TAutoConsoleVariable<int32> CVarDetailedMipAlphaLogging(
+	TEXT("r.DetailedMipAlphaLogging"),
+	0,
+	TEXT("Prints extra log messages for tracking when alpha gets removed/introduced during texture")
+	TEXT("image processing.")
+);
+
 DEFINE_LOG_CATEGORY_STATIC(LogTextureCompressor, Log, All);
 
 /*------------------------------------------------------------------------------
@@ -1689,7 +1696,7 @@ static void LinearizeToWorkingColorSpace(const FImage& SrcImage, FImage& DstImag
 		SrcImageView.GammaSpace = DstImage.GammaSpace; // EGammaSpace::Linear
 	}
 
-        // CopyImage to get pixels in RGAB32F , then OCIO will act on those in-place in DstImage
+	// CopyImage to get pixels in RGAB32F , then OCIO will act on those in-place in DstImage
 	FImageCore::CopyImage(SrcImageView, DstImage);
 
 	// Decode and/or color transform to the working color space when needed
@@ -2959,6 +2966,89 @@ void ITextureCompressorModule::AdjustImageColors(FImage& Image, const FTextureBu
 	}
 }
 
+bool ITextureCompressorModule::DetermineAlphaChannelTransparency(const FTextureBuildSettings& InBuildSettings, const FLinearColor& InChannelMin, const FLinearColor& InChannelMax, bool& bOutAlphaIsTransparent)
+{
+	// Settings that affect alpha:
+	// 1. ColorTransforms - if they have a color transform we assume it's not affecting alpha (this is the UE
+	//		integration assumption separate from this) this only gets dicey if we have ReplicateRed. Probably
+	//		the best thing to do is just assume the transform doesn't affects opaqueness and let them use the
+	//		force overrides if it's wrong.
+	// 2. ForceAlphaChannel - Note ForceNoAlpha forces BC1 before us so we don't see it.
+	// 3. AdjustMinAlpha, AdjustMaxAlpha
+	// 4. ReplicateRed.. which means we need to track all operations on the red channel. Can probably just call
+	//		AdjustImageColors directly on our boundaries and see where they go?
+	// 5. ChromaKey! I think if they have this on we assume it matches SOMEWHERE and we need alpha.
+
+	// ForceNo gets applied first because it happens at the texture format selection phase, so it effectively
+	// overrides everything as AutoDXT gets cleared.
+	if (InBuildSettings.bForceNoAlphaChannel)
+	{
+		bOutAlphaIsTransparent = false; // note we don't see this with autodxt as it gets forced to BC1 early.
+		return true;
+	}
+
+	if (InBuildSettings.bForceAlphaChannel ||
+		InBuildSettings.bChromaKeyTexture // If they have a chroma key they are expecting transparency!
+		)
+	{
+		bOutAlphaIsTransparent = true;
+		return true;
+	}
+
+	// No forces - try and predict how the color transforms will affect the colors.
+	if (InBuildSettings.bHasColorSpaceDefinition)
+	{
+		// We assert that color spaces don't affect alpha per our integration with OCIO. So, the only way 
+		// this affects us is if RepliateRed is set. 
+		// 
+		// We assume the color space doesn't change the opaquenss of the red channel so we can pass through to
+		// the normal replicate red behavior. This could definitely fail, but we don't really expect anyone
+		// to combine a color space remap _and_ replicate red, so if they hit this they can get what they want
+		// with Force/ForceNoAlpha.
+		if (InBuildSettings.bReplicateRed)
+		{
+			return false;
+		}
+	}
+
+	if (InBuildSettings.bComputeBokehAlpha)
+	{
+		// Bokeh stores information in the alpha channel during processing - so we need an alpha channel
+		bOutAlphaIsTransparent = true;
+		return true;
+	}
+
+	FLinearColor ChannelMinMax[2] = {InChannelMin, InChannelMax};
+
+	// If there's an image adjustment, run it on the colors so we see what happens.
+	// This also handles AdjustMinAlpha/MaxAlpha.
+	if (NeedAdjustImageColors(InBuildSettings))
+	{
+		if (InBuildSettings.bUseNewMipFilter)
+		{
+			AdjustColorsNew(ChannelMinMax, 2, InBuildSettings);
+		}
+		else
+		{
+			AdjustColorsOld(ChannelMinMax, 2, InBuildSettings);
+		}
+	}
+
+	if (InBuildSettings.bReplicateRed)
+	{
+		ChannelMinMax[0].A = ChannelMinMax[0].R;
+		ChannelMinMax[1].A = ChannelMinMax[1].R;
+	}
+
+	bOutAlphaIsTransparent = false;
+	if (ChannelMinMax[0].A < 1 ||
+		ChannelMinMax[1].A < 1)
+	{
+		bOutAlphaIsTransparent = true;
+	}
+	return true;
+}
+
 /**
  * Compute the alpha channel how BokehDOF needs it setup
  *
@@ -3771,6 +3861,8 @@ public:
 
 			return false;
 		}
+
+		const bool bDoDetailedAlphaLogging = CVarDetailedMipAlphaLogging.GetValueOnAnyThread() != 0;
 		
 		// @todo Oodle: option to dump the Source image here
 		//		we have dump in TextureFormatOodle for the after-processing (before encoding) image
@@ -3779,9 +3871,14 @@ public:
 		TArray<FImage> IntermediateMipChain;
 
 		bool bSourceMipsAlphaDetected = false;
-		if (OutMetadata)
+		if (OutMetadata || bDoDetailedAlphaLogging)
 		{
 			bSourceMipsAlphaDetected = FImageCore::DetectAlphaChannel(SourceMips[0]);
+			
+			if (bDoDetailedAlphaLogging)
+			{
+				UE_LOG(LogTextureCompressor, Display, TEXT("[alpha] Source Mips: %d - %.*s"), bSourceMipsAlphaDetected, DebugTexturePathName.Len(), DebugTexturePathName.GetData());
+			}
 		}
 
 		// allow to leave texture in sRGB in case compressor accepts other than non-F32 input source
@@ -3832,17 +3929,61 @@ public:
 
 			if (!ApplyCompositeTextureToMips(IntermediateMipChain, IntermediateAssociatedNormalSourceMipChain, BuildSettings.CompositeTextureMode, BuildSettings.CompositePower, BuildSettings.LODBias))
 			{
-				UE_LOG(LogTextureCompressor, Warning, TEXT("ApplyCompositeTextureToMips failed [%.*s]"),
+				UE_LOG(LogTextureCompressor, Display, TEXT("ApplyCompositeTextureToMips failed [%.*s]"),
 					DebugTexturePathName.Len(),DebugTexturePathName.GetData());
 
 				return false;
 			}
+
+			if (bDoDetailedAlphaLogging)
+			{
+				UE_LOG(LogTextureCompressor, Display, TEXT("[alpha] Associated Normal Mips: %d - %.*s"), FImageCore::DetectAlphaChannel(IntermediateMipChain[0]),
+					DebugTexturePathName.Len(), DebugTexturePathName.GetData());
+			}
+		}
+
+
+		bool bIntermediateMipsAlphaDetected = FImageCore::DetectAlphaChannel(IntermediateMipChain[0]);
+		if (bDoDetailedAlphaLogging)
+		{
+			UE_LOG(LogTextureCompressor, Display, TEXT("[alpha] Intermediate Mips: %d - %.*s"), bIntermediateMipsAlphaDetected,
+				DebugTexturePathName.Len(), DebugTexturePathName.GetData());
+
+			if (bIntermediateMipsAlphaDetected != bSourceMipsAlphaDetected)
+			{
+				UE_LOG(LogTextureCompressor, Display, TEXT("[alpha] Source / Intermediate diff (force %d force no %d) - %.*s"),
+					BuildSettings.bForceAlphaChannel, 
+					BuildSettings.bForceNoAlphaChannel, 
+					DebugTexturePathName.Len(),
+					DebugTexturePathName.GetData());
+			}
 		}
 
-		
 		// DetectAlphaChannel on the top mip of the generated mip chain. SoonTM this will use the source mip chain. Testing has
 		// shown this to be 99.9% the same, and allows us to get the pixel format earlier.
-		const bool bImageHasAlphaChannel = BuildSettings.GetTextureExpectsAlphaInPixelFormat(FImageCore::DetectAlphaChannel(IntermediateMipChain[0]));
+		const bool bImageHasAlphaChannel = BuildSettings.GetTextureExpectsAlphaInPixelFormat(bIntermediateMipsAlphaDetected);
+
+		// If we know what our transparency is make sure it matches what we expect.
+		if (BuildSettings.bKnowAlphaTransparency)
+		{
+			if (BuildSettings.bHasTransparentAlpha != bImageHasAlphaChannel)
+			{
+				// We don't actually use this yet but we DO need to know when they aren't matching to hunt
+				// down bugs so we do a warning. The actual build is still good and this doesn't affect the output.
+
+				// We only actually care if it affects the output format.
+				FName ActualTextureFormat = UE::TextureBuildUtilities::TextureFormatRemovePrefixFromName(BuildSettings.TextureFormatName);
+				if (ActualTextureFormat == "AutoDXT")
+				{
+					UE_LOG(LogTextureCompressor, Warning, TEXT("Expected texture alpha didn't match reality: Expected: %s Reality: %s Texture: %.*s"),
+						BuildSettings.bHasTransparentAlpha ? TEXT("true") : TEXT("false"),
+						bImageHasAlphaChannel ? TEXT("true") : TEXT("false"),
+						DebugTexturePathName.Len(),
+						DebugTexturePathName.GetData());
+				}
+			}
+		}
+
 		
 		UE::Tasks::TTask<uint64> HashingTask;
 		TArray<FImageInfo> SaveImageInfos;
@@ -3948,6 +4089,8 @@ private:
 			}
 		}
 
+		const bool bDoDetailedAlphaLogging = CVarDetailedMipAlphaLogging.GetValueOnAnyThread() != 0;
+
 		// handling of bLongLatCubemap seems overly complicated
 		//	what it should do is convert it right at the start here
 		//	then treat it as a standard cubemap below, no special cases
@@ -4253,6 +4396,11 @@ private:
 		{
 			const FImage& Image = (*pSourceMips)[MipIndex];
 
+			if (bDoDetailedAlphaLogging)
+			{
+				UE_LOG(LogTextureCompressor, Display, TEXT("[alpha] Pre Process: %d - %.*s"), FImageCore::DetectAlphaChannel(Image), DebugTexturePathName.Len(), DebugTexturePathName.GetData());
+			}
+			
 			// copy mips over + processing
 			// this is a code dupe of the processing done in GenerateMipChain
 
@@ -4263,7 +4411,6 @@ private:
 			{
 				// Generate the base mip from the long-lat source image.
 				GenerateBaseCubeMipFromLongitudeLatitude2D(&Mip, Image, BuildSettings);
-	
 				check( CopyCount == 1 );
 			}
 			else
@@ -4279,6 +4426,11 @@ private:
 					}
 
 					GenerateTopMip(Temp, Mip, BuildSettings);
+
+					if (bDoDetailedAlphaLogging)
+					{
+						UE_LOG(LogTextureCompressor, Display, TEXT("[alpha] ApplyKernel: %d - %.*s"), FImageCore::DetectAlphaChannel(Mip), DebugTexturePathName.Len(), DebugTexturePathName.GetData());
+					}
 				}
 				else
 				{
@@ -4289,6 +4441,11 @@ private:
 						{
 							NormalizeMip(Mip);
 						}
+
+						if (bDoDetailedAlphaLogging)
+						{
+							UE_LOG(LogTextureCompressor, Display, TEXT("[alpha] Linearize: %d - %.*s"), FImageCore::DetectAlphaChannel(Mip), DebugTexturePathName.Len(), DebugTexturePathName.GetData());
+						}
 					}
 					else
 					{
@@ -4302,7 +4459,11 @@ private:
 							DestGammaSpace = EGammaSpace::Linear;
 						}
 						Image.CopyTo(Mip, DestFormat, DestGammaSpace);
-
+						
+						if (bDoDetailedAlphaLogging)
+						{
+							UE_LOG(LogTextureCompressor, Display, TEXT("[alpha] Copy: %d - %.*s"), FImageCore::DetectAlphaChannel(Mip), DebugTexturePathName.Len(), DebugTexturePathName.GetData());
+						}
 						//@@CB todo : when Mip format == Image format, we can Move instead of Copy
 						//	have to make sure that's okay with SourceMips/TextureData
 					}
@@ -4314,15 +4475,30 @@ private:
 			if (BuildSettings.Downscale > 1.f)
 			{		
 				DownscaleImage(Mip, Mip, FTextureDownscaleSettings(BuildSettings));
+
+				if (bDoDetailedAlphaLogging)
+				{
+					UE_LOG(LogTextureCompressor, Display, TEXT("[alpha] Downscale: %d - %.*s"), FImageCore::DetectAlphaChannel(Mip), DebugTexturePathName.Len(), DebugTexturePathName.GetData());
+				}
 			}
 
 			// Apply color adjustments
 			AdjustImageColors(Mip, BuildSettings);
 
+			if (bDoDetailedAlphaLogging)
+			{
+				UE_LOG(LogTextureCompressor, Display, TEXT("[alpha] ColorAdjust: %d - %.*s"), FImageCore::DetectAlphaChannel(Mip), DebugTexturePathName.Len(), DebugTexturePathName.GetData());
+			}
+
 			if (BuildSettings.bComputeBokehAlpha)
 			{
 				// To get the occlusion in the BokehDOF shader working for all Bokeh textures.
 				ComputeBokehAlpha(Mip);
+
+				if (bDoDetailedAlphaLogging)
+				{
+					UE_LOG(LogTextureCompressor, Display, TEXT("[alpha] Bokeh: %d - %.*s"), FImageCore::DetectAlphaChannel(Mip), DebugTexturePathName.Len(), DebugTexturePathName.GetData());
+				}
 			}
 			if (BuildSettings.bFlipGreenChannel)
 			{
@@ -4379,6 +4555,12 @@ private:
 				NumOutputMips, OutMipChain[0].SizeX, OutMipChain[0].SizeY, OutMipChain[0].NumSlices);
 		}
 
+		if (bDoDetailedAlphaLogging)
+		{
+			UE_LOG(LogTextureCompressor, Display, TEXT("[alpha] BeginPost: %d - %.*s"), FImageCore::DetectAlphaChannel(OutMipChain[0]), DebugTexturePathName.Len(), DebugTexturePathName.GetData());
+		}
+
+
 		// Apply post-mip generation adjustments.
 		if ( BuildSettings.bNormalizeNormals )
 		{
@@ -4414,6 +4596,11 @@ private:
 		{
 			check( !BuildSettings.bReplicateAlpha ); // cannot both be set 
 			ReplicateRedChannel(OutMipChain);
+
+			if (bDoDetailedAlphaLogging)
+			{
+				UE_LOG(LogTextureCompressor, Display, TEXT("[alpha] ReplicateRed: %d - %.*s"), FImageCore::DetectAlphaChannel(OutMipChain[0]), DebugTexturePathName.Len(), DebugTexturePathName.GetData());
+			}
 		}
 		else if (BuildSettings.bReplicateAlpha)
 		{
@@ -4425,6 +4612,11 @@ private:
 			ApplyYCoCgBlockScale(OutMipChain);
 		}
 
+		if (bDoDetailedAlphaLogging)
+		{
+			UE_LOG(LogTextureCompressor, Display, TEXT("[alpha] End: %d - %.*s"), FImageCore::DetectAlphaChannel(OutMipChain[0]), DebugTexturePathName.Len(), DebugTexturePathName.GetData());
+		}
+
 		return true;
 	}
 
diff --git a/Engine/Source/Developer/TextureCompressor/Public/TextureCompressorModule.h b/Engine/Source/Developer/TextureCompressor/Public/TextureCompressorModule.h
--- a/Engine/Source/Developer/TextureCompressor/Public/TextureCompressorModule.h
+++ b/Engine/Source/Developer/TextureCompressor/Public/TextureCompressorModule.h
@@ -251,6 +251,13 @@ struct FTextureBuildSettings
 	// to TextureFormatName.
 	FName BaseTextureFormatName;
 	
+	// Whether bHasTransparentAlpha is valid.
+	bool bKnowAlphaTransparency = false;
+
+	// Only valid if bKnowAlphaTransparency is true. This is whether the resulting texture is expected to require
+	// an alpha channel based on scanning the source mips and analyzing the build settings.
+	bool bHasTransparentAlpha = false;
+
 	static constexpr uint32 MaxTextureResolutionDefault = TNumericLimits<uint32>::Max();
 
 	/** Default settings. */
@@ -413,6 +420,18 @@ public:
 		uint32 MipChainDepth
 		);
 
+	/**
+	* Given the channel min/max for the top mip of a texture's source, determine whether there will be any transparency
+	* after image processing occurs, and thus the texture will need to be BC3 for the purposes of AutoDXT format selection.
+	* 
+	* @param bOutAlphaIsTransparent		*This is undetermined if the return is false!*. If true, we expect the image after
+	*									processing to require an alpha channel.
+	* @return							Whether the alpha channel can be determined. There are plenty of edge cases where 
+	*									it's not feasible to determine the result prior to full processing, however these
+	*									are rare in practice.
+	*/
+	TEXTURECOMPRESSOR_API static bool DetermineAlphaChannelTransparency(const FTextureBuildSettings& InBuildSettings, const FLinearColor& InChannelMin, const FLinearColor& InChannelMax, bool& bOutAlphaIsTransparent);
+
 	/**
      * Adjusts the colors of the image using the specified settings
      *
diff --git a/Engine/Source/Runtime/Engine/Private/TextureDerivedData.cpp b/Engine/Source/Runtime/Engine/Private/TextureDerivedData.cpp
--- a/Engine/Source/Runtime/Engine/Private/TextureDerivedData.cpp
+++ b/Engine/Source/Runtime/Engine/Private/TextureDerivedData.cpp
@@ -652,6 +652,15 @@ static void FinalizeBuildSettingsForLayer(
 		OutSettings.bReplicateRed = true;
 	}
 
+	// If we have channel boundary information, use that to determine whether we expect to have
+	// a non opaque alpha.
+	if (LayerIndex < Texture.Source.GetLayerColorInfo().Num())
+	{
+		const FTextureSourceLayerColorInfo& LayerChannelBounds = Texture.Source.GetLayerColorInfo()[LayerIndex];
+		OutSettings.bKnowAlphaTransparency = ITextureCompressorModule::DetermineAlphaChannelTransparency(OutSettings, 
+			LayerChannelBounds.ColorMin, LayerChannelBounds.ColorMax, OutSettings.bHasTransparentAlpha);
+	}
+
 	// this is called once per Texture with OutSettings.TextureFormatName == None
 	//	and then called again (per Layer) with OutSettings.TextureFormatName filled out
 
diff --git a/Engine/Source/Runtime/Engine/Private/TextureDerivedDataTask.cpp b/Engine/Source/Runtime/Engine/Private/TextureDerivedDataTask.cpp
--- a/Engine/Source/Runtime/Engine/Private/TextureDerivedDataTask.cpp
+++ b/Engine/Source/Runtime/Engine/Private/TextureDerivedDataTask.cpp
@@ -209,6 +209,18 @@ void FTextureSourceData::Init(UTexture& InTexture, TextureMipGenSettings InMipGe
 {
 	check( bValid == false ); // we set to true at the end, acts as our return value
 
+	// Copy the channel min/max if we have it to avoid redoing it.
+	if (InTexture.Source.GetLayerColorInfo().Num())
+	{
+		LayerChannelMinMax.Reset();
+		for (const FTextureSourceLayerColorInfo& LayerColorInfo : InTexture.Source.GetLayerColorInfo())
+		{
+			TPair<FLinearColor, FLinearColor>& MinMax = LayerChannelMinMax.AddDefaulted_GetRef();
+			MinMax.Key = LayerColorInfo.ColorMin;
+			MinMax.Value = LayerColorInfo.ColorMax;
+		}
+	}
+
 	const int32 NumBlocks = InTexture.Source.GetNumBlocks();
 	const int32 NumLayers = InTexture.Source.GetNumLayers();
 	if (NumBlocks < 1 || NumLayers < 1)
@@ -303,7 +315,26 @@ void FTextureSourceData::GetSourceMips(FTextureSource& Source, IImageWrapperModu
 			return;
 		}
 
-		const FTextureSource::FMipData ScopedMipData = Source.GetMipData(InImageWrapper);
+		// This locks the entire source - not just mip 0, so we can freely lock other mips without copies.
+		FTextureSource::FMipLock LockedTextureSource(FTextureSource::ELockState::ReadOnly, &Source, 0);
+		if (LockedTextureSource.IsValid() == false)
+		{
+			UE_LOG(LogTexture, Warning, TEXT("Cannot lock texture source data for %s"), *TextureFullName);
+			ReleaseMemory();
+			bValid = false;
+			return;
+		}
+
+		// If we didn't get this from the texture source
+		if (LayerChannelMinMax.Num() != Layers.Num())
+		{
+			LayerChannelMinMax.Reset();
+			for (int32 LayerIndex = 0; LayerIndex < Layers.Num(); ++LayerIndex)
+			{
+				TPair<FLinearColor, FLinearColor>& LayerInfo = LayerChannelMinMax.AddDefaulted_GetRef();
+				Source.ComputeChannelLinearMinMax(LayerIndex, LayerInfo.Key, LayerInfo.Value);
+			}
+		}
 
 		{
 			TRACE_CPUPROFILER_EVENT_SCOPE(FTextureSourceData::GetSourceMips_CopyMips);
@@ -329,13 +360,16 @@ void FTextureSourceData::GetSourceMips(FTextureSource& Source, IImageWrapperModu
 								LayerData.SourceGammaSpace
 							);
 
-							if (!ScopedMipData.GetMipData(SourceMip.RawData, BlockIndex, LayerIndex, MipIndex))
-							{
-								UE_LOG(LogTexture, Warning, TEXT("Cannot retrieve source data for mip %d of %s"), MipIndex, *TextureFullName);
-								ReleaseMemory();
-								bValid = false;
-								break;
-							}
+							// This gets a view into the decompressed data so we need to copy out the mip we care about.
+							FTextureSource::FMipLock LockedMipData(FTextureSource::ELockState::ReadOnly, &Source, BlockIndex, LayerIndex, MipIndex);
+
+							check(MipSizeX == LockedMipData.Image.SizeX);
+							check(MipSizeY == LockedMipData.Image.SizeY);
+							check(MipSizeZ == LockedMipData.Image.NumSlices);
+							check(SourceMip.RawData.Num() == LockedMipData.GetDataSize());
+
+							SourceMip.RawData.Reset(LockedMipData.GetDataSize());
+							SourceMip.RawData.Append((const uint8*)LockedMipData.GetRawData(), LockedMipData.GetDataSize());
 
 							MipSizeX = FMath::Max(MipSizeX / 2, 1);
 							MipSizeY = FMath::Max(MipSizeY / 2, 1);
@@ -2015,6 +2049,18 @@ void FTextureCacheDerivedDataWorker::DoWork()
 	{
 		if (DDC1_LoadAndValidateTextureData(Texture, TextureData, CompositeTextureData, ImageWrapper, bAllowAsyncLoading))
 		{
+			for (int32 LayerIndex = 0; LayerIndex < BuildSettingsPerLayerFetchOrBuild.Num(); LayerIndex++)
+			{
+				if (LayerIndex < TextureData.LayerChannelMinMax.Num())
+				{
+					BuildSettingsPerLayerFetchOrBuild[LayerIndex].bKnowAlphaTransparency = Compressor->DetermineAlphaChannelTransparency(
+						BuildSettingsPerLayerFetchOrBuild[LayerIndex], 
+						TextureData.LayerChannelMinMax[LayerIndex].Key,
+						TextureData.LayerChannelMinMax[LayerIndex].Value,
+						BuildSettingsPerLayerFetchOrBuild[LayerIndex].bHasTransparentAlpha);
+				}
+			}
+
 			// Replace any existing DDC data, if corrupt compression was detected
 			const bool bReplaceExistingDDC = bInvalidVirtualTextureCompression;
 
diff --git a/Engine/Source/Runtime/Engine/Private/TextureDerivedDataTask.h b/Engine/Source/Runtime/Engine/Private/TextureDerivedDataTask.h
--- a/Engine/Source/Runtime/Engine/Private/TextureDerivedDataTask.h
+++ b/Engine/Source/Runtime/Engine/Private/TextureDerivedDataTask.h
@@ -115,6 +115,8 @@ struct FTextureSourceData
 		bValid = false;
 	}
 
+	TArray<TPair<FLinearColor, FLinearColor>> LayerChannelMinMax;
+
 	FString TextureFullName;
 	FTextureSource AsyncSource;
 	TArray<FTextureSourceLayerData> Layers;
