Texture fix GetBuildRequiredMemoryEstimate was not accounting for pow2 pad options
causing it to under-estimate memory use in some cases
move GetPowerOfTwoTargetTextureSize to TextureBuildUtilities so it can be shared
fix some comments

diff --git a/Engine/Source/Developer/TextureBuildUtilities/Private/TextureBuildUtilities.cpp b/Engine/Source/Developer/TextureBuildUtilities/Private/TextureBuildUtilities.cpp
--- a/Engine/Source/Developer/TextureBuildUtilities/Private/TextureBuildUtilities.cpp
+++ b/Engine/Source/Developer/TextureBuildUtilities/Private/TextureBuildUtilities.cpp
@@ -289,5 +289,66 @@ void GetPlaceholderTextureImage(FImage* OutImage)
 
 }
 
+
+// Returns true if the target texture size is different and padding/stretching is required.
+TEXTUREBUILDUTILITIES_API bool GetPowerOfTwoTargetTextureSize(int32 InMip0SizeX, int32 InMip0SizeY, int32 InMip0NumSlices, bool bInIsVolume, ETexturePowerOfTwoSetting::Type InPow2Setting, int32 InResizeDuringBuildX, int32 InResizeDuringBuildY, int32& OutTargetSizeX, int32& OutTargetSizeY, int32& OutTargetSizeZ)
+{
+	int32 TargetTextureSizeX = InMip0SizeX;
+	int32 TargetTextureSizeY = InMip0SizeY;
+	int32 TargetTextureSizeZ = bInIsVolume ? InMip0NumSlices : 1; // Only used for volume texture.
+
+	const int32 PowerOfTwoTextureSizeX = FMath::RoundUpToPowerOfTwo(TargetTextureSizeX);
+	const int32 PowerOfTwoTextureSizeY = FMath::RoundUpToPowerOfTwo(TargetTextureSizeY);
+	const int32 PowerOfTwoTextureSizeZ = FMath::RoundUpToPowerOfTwo(TargetTextureSizeZ);
+
+	switch (InPow2Setting)
+	{
+	case ETexturePowerOfTwoSetting::None:
+		break;
+
+	case ETexturePowerOfTwoSetting::PadToPowerOfTwo:
+	case ETexturePowerOfTwoSetting::StretchToPowerOfTwo:
+		TargetTextureSizeX = PowerOfTwoTextureSizeX;
+		TargetTextureSizeY = PowerOfTwoTextureSizeY;
+		TargetTextureSizeZ = PowerOfTwoTextureSizeZ;
+		break;
+
+	case ETexturePowerOfTwoSetting::PadToSquarePowerOfTwo:
+	case ETexturePowerOfTwoSetting::StretchToSquarePowerOfTwo:
+		TargetTextureSizeX = TargetTextureSizeY = TargetTextureSizeZ =
+			FMath::Max3<int32>(PowerOfTwoTextureSizeX, PowerOfTwoTextureSizeY, PowerOfTwoTextureSizeZ);
+		break;
+
+	case ETexturePowerOfTwoSetting::ResizeToSpecificResolution:
+		if (InResizeDuringBuildX)
+		{
+			TargetTextureSizeX = InResizeDuringBuildX;
+		}
+		if (InResizeDuringBuildY)
+		{
+			TargetTextureSizeY = InResizeDuringBuildY;
+		}
+		break;
+
+	default:
+		checkf(false, TEXT("Unknown entry in ETexturePowerOfTwoSetting::Type"));
+		break;
+	}
+
+	// Z only matters as a sampling dimension if we are a volume texture.
+	if (bInIsVolume == false)
+	{
+		TargetTextureSizeZ = InMip0NumSlices;
+	}
+
+	OutTargetSizeX = TargetTextureSizeX;
+	OutTargetSizeY = TargetTextureSizeY;
+	OutTargetSizeZ = TargetTextureSizeZ;
+
+	return (TargetTextureSizeX != InMip0SizeX) ||
+		(TargetTextureSizeY != InMip0SizeY) ||
+		(bInIsVolume && TargetTextureSizeZ != InMip0NumSlices);
+}
+
 } // namespace
 }
\ No newline at end of file
diff --git a/Engine/Source/Developer/TextureBuildUtilities/Public/TextureBuildUtilities.h b/Engine/Source/Developer/TextureBuildUtilities/Public/TextureBuildUtilities.h
--- a/Engine/Source/Developer/TextureBuildUtilities/Public/TextureBuildUtilities.h
+++ b/Engine/Source/Developer/TextureBuildUtilities/Public/TextureBuildUtilities.h
@@ -5,6 +5,7 @@
 #include "CoreMinimal.h"
 #include "ImageCore.h"
 #include "Interfaces/ITextureFormat.h"
+#include "Engine/TextureDefines.h"
 
 struct FTextureBuildSettings;
 
@@ -81,5 +82,9 @@ TEXTUREBUILDUTILITIES_API ERawImageFormat::Type GetVirtualTextureBuildIntermedia
 TEXTUREBUILDUTILITIES_API void GetPlaceholderTextureImageInfo(FImageInfo* OutImageInfo);
 TEXTUREBUILDUTILITIES_API void GetPlaceholderTextureImage(FImage* OutImage);
 
+// Returns true if the target texture size is different and padding/stretching is required.
+//	if InPow2Setting == None, the Out sizes match the In sizes, and false is returned
+TEXTUREBUILDUTILITIES_API bool GetPowerOfTwoTargetTextureSize(int32 InMip0SizeX, int32 InMip0SizeY, int32 InMip0NumSlices, bool bInIsVolume, ETexturePowerOfTwoSetting::Type InPow2Setting, int32 InResizeDuringBuildX, int32 InResizeDuringBuildY, int32& OutTargetSizeX, int32& OutTargetSizeY, int32& OutTargetSizeZ);
+
 } // namespace TextureBuildUtilities
 } // namespace UE
diff --git a/Engine/Source/Developer/TextureBuildUtilities/TextureBuildUtilities.Build.cs b/Engine/Source/Developer/TextureBuildUtilities/TextureBuildUtilities.Build.cs
--- a/Engine/Source/Developer/TextureBuildUtilities/TextureBuildUtilities.Build.cs
+++ b/Engine/Source/Developer/TextureBuildUtilities/TextureBuildUtilities.Build.cs
@@ -11,7 +11,8 @@ public class TextureBuildUtilities : ModuleRules
 		// Include only , no link :
 		PrivateIncludePathModuleNames.AddRange(new string[]
 		{
-			"TextureCompressor"
+			"TextureCompressor",
+			"Engine"
 		});
 
 		PrivateDependencyModuleNames.AddRange(
diff --git a/Engine/Source/Developer/TextureCompressor/Private/TextureCompressorModule.cpp b/Engine/Source/Developer/TextureCompressor/Private/TextureCompressorModule.cpp
--- a/Engine/Source/Developer/TextureCompressor/Private/TextureCompressorModule.cpp
+++ b/Engine/Source/Developer/TextureCompressor/Private/TextureCompressorModule.cpp
@@ -1519,6 +1519,9 @@ static void DownscaleImage(const FImage& SrcImage, FImage& DstImage, const FText
 		return;
 	}
 	
+	// legacy/bad path , not used for new textures that have UseNewMipFilter set
+	//	left as-is to prevent changing old content
+
 	// what this function does is 2X downsamples with the mip filter
 	//	and then a final bilinear resize to the desired final size
 	// instead just use ResizeImage to go directly from source size to final size in one step
@@ -1612,8 +1615,8 @@ static void DownscaleImage(const FImage& SrcImage, FImage& DstImage, const FText
 	FImageView2D SrcImageData(*ImageChain[0], 0);
 	FImageView2D DstImageData(*ImageChain[1], 0);
 				
-	// @todo OodleImageResize : not sure this is a correct image resize without shift; does it get pixel center offsets right?
-	//	use FImageCore::ResizeImage here instead
+	// this is not a correct image resize without shift; does not get pixel center offsets right
+	//	this is in the legacy/deprecated path so leave as-is ; new path is correct
 	for (int32 Y = 0; Y < FinalSizeY; ++Y)
 	{
 		float SourceY = (float)Y * Downscale;
@@ -3588,68 +3591,6 @@ static void NormalizeMip(FImage& InOutMip)
 }
 
 
-// Returns true if the target texture size is different and padding/stretching is required.
-static bool GetPowerOfTwoTargetTextureSize(int32 InMip0SizeX, int32 InMip0SizeY, int32 InMip0NumSlices, bool bInIsVolume, ETexturePowerOfTwoSetting::Type InPow2Setting, int32 InResizeDuringBuildX, int32 InResizeDuringBuildY, int32& OutTargetSizeX, int32& OutTargetSizeY, int32& OutTargetSizeZ)
-{
-	check(InPow2Setting != ETexturePowerOfTwoSetting::None);
-
-	int32 TargetTextureSizeX = InMip0SizeX;
-	int32 TargetTextureSizeY = InMip0SizeY;
-	int32 TargetTextureSizeZ = bInIsVolume ? InMip0NumSlices : 1; // Only used for volume texture.
-
-	const int32 PowerOfTwoTextureSizeX = FMath::RoundUpToPowerOfTwo(TargetTextureSizeX);
-	const int32 PowerOfTwoTextureSizeY = FMath::RoundUpToPowerOfTwo(TargetTextureSizeY);
-	const int32 PowerOfTwoTextureSizeZ = FMath::RoundUpToPowerOfTwo(TargetTextureSizeZ);
-
-	switch (InPow2Setting)
-	{
-	// None should not get here
-
-	case ETexturePowerOfTwoSetting::PadToPowerOfTwo:
-	case ETexturePowerOfTwoSetting::StretchToPowerOfTwo:
-		TargetTextureSizeX = PowerOfTwoTextureSizeX;
-		TargetTextureSizeY = PowerOfTwoTextureSizeY;
-		TargetTextureSizeZ = PowerOfTwoTextureSizeZ;
-		break;
-
-	case ETexturePowerOfTwoSetting::PadToSquarePowerOfTwo:
-	case ETexturePowerOfTwoSetting::StretchToSquarePowerOfTwo:
-		TargetTextureSizeX = TargetTextureSizeY = TargetTextureSizeZ =
-			FMath::Max3<int32>(PowerOfTwoTextureSizeX, PowerOfTwoTextureSizeY, PowerOfTwoTextureSizeZ);
-		break;
-
-	case ETexturePowerOfTwoSetting::ResizeToSpecificResolution:
-		if (InResizeDuringBuildX)
-		{
-			TargetTextureSizeX = InResizeDuringBuildX;
-		}
-		if (InResizeDuringBuildY)
-		{
-			TargetTextureSizeY = InResizeDuringBuildY;
-		}
-		break;
-
-	default:
-		checkf(false, TEXT("Unknown entry in ETexturePowerOfTwoSetting::Type"));
-		break;
-	}
-
-	// Z only matters as a sampling dimension if we are a volume texture.
-	if (bInIsVolume == false)
-	{
-		TargetTextureSizeZ = InMip0NumSlices;
-	}
-
-	OutTargetSizeX = TargetTextureSizeX;
-	OutTargetSizeY = TargetTextureSizeY;
-	OutTargetSizeZ = TargetTextureSizeZ;
-
-	return (TargetTextureSizeX != InMip0SizeX) ||
-		(TargetTextureSizeY != InMip0SizeY) ||
-		(bInIsVolume && TargetTextureSizeZ != InMip0NumSlices);
-}
-
-
 int32 ITextureCompressorModule::GetMipCountForBuildSettings(
 	int32 InMip0SizeX, int32 InMip0SizeY, int32 InMip0NumSlices,
 	int32 InExistingMipCount,
@@ -3681,7 +3622,7 @@ int32 ITextureCompressorModule::GetMipCountForBuildSettings(
 			PowerOfTwoMode != ETexturePowerOfTwoSetting::None)
 		{
 			int32 TargetSizeX, TargetSizeY, TargetSizeZ;
-			bool NeedsAdjustment = GetPowerOfTwoTargetTextureSize(BaseSizeX, BaseSizeY, BaseSizeZ, BuildSettings.bVolume, PowerOfTwoMode, BuildSettings.ResizeDuringBuildX, BuildSettings.ResizeDuringBuildY, TargetSizeX, TargetSizeY, TargetSizeZ);
+			bool NeedsAdjustment = UE::TextureBuildUtilities::GetPowerOfTwoTargetTextureSize(BaseSizeX, BaseSizeY, BaseSizeZ, BuildSettings.bVolume, PowerOfTwoMode, BuildSettings.ResizeDuringBuildX, BuildSettings.ResizeDuringBuildY, TargetSizeX, TargetSizeY, TargetSizeZ);
 			if (NeedsAdjustment)
 			{
 				// In this case we are regenerating the entire mip chain.
@@ -3792,9 +3733,6 @@ public:
 		TArray<UE::Tasks::TTask<FXxHash64>, TInlineAllocator<16>> MipHashTasks;
 		MipHashTasks.Reserve(IntermediateMipChain.Num());
 
-		//double start_time = FPlatformTime::Seconds();
-		//uint64 raw_bytes_hashed = 0;
-
 		// Hash the mips before we compress them. This gets saved as part of the derived data and then added to the
 		// diff tags during cook so we can catch determinism issues.
 		FXxHash64Builder MipHashBuilder;
@@ -3811,10 +3749,6 @@ public:
 			check( Mip.RawData.Num() != 0 );
 
 			MipHashTasks.Add(UE::Tasks::Launch(TEXT("ComputeMipChainHash"), [&Mip] { return FXxHash64::HashBufferChunked(MakeMemoryView(Mip.RawData), 256 << 10); }));
-
-			//raw_bytes_hashed += Mip.RawData.Num();
-
-			//Blake3 is ~3.5 GB/s , around 4X slower than xxHash.
 		}
 
 		for (UE::Tasks::TTask<FXxHash64>& HashTask : MipHashTasks)
@@ -3823,11 +3757,6 @@ public:
 			MipHashBuilder.Update(&MipHash.Hash, sizeof(MipHash.Hash));
 		}
 
-		//double end_time = FPlatformTime::Seconds();
-
-		//UE_LOG(LogTextureCompressor, Display, TEXT("Hashed %lld bytes in %.3f millis = %.3f GB/s"), raw_bytes_hashed, (end_time - start_time) * 1000, (raw_bytes_hashed / (end_time - start_time)) / (1000 * 1000 * 1000));
-		// Threaded chunked xxHash with 256kb blocks ~120 GB / s, has a lot of tiny hashes in this timing scope though.
-
 		return MipHashBuilder.Finalize().Hash;
 	}
 
@@ -4102,7 +4031,7 @@ private:
 			int32 TargetTextureSizeX = 0;
 			int32 TargetTextureSizeY = 0;
 			int32 TargetTextureSizeZ = 0;			
-			bool bPadOrStretchTexture = GetPowerOfTwoTargetTextureSize(
+			bool bPadOrStretchTexture = UE::TextureBuildUtilities::GetPowerOfTwoTargetTextureSize(
 				FirstSourceMipImage.SizeX, FirstSourceMipImage.SizeY, FirstSourceMipImage.NumSlices,
 				BuildSettings.bVolume, PowerOfTwoMode, BuildSettings.ResizeDuringBuildX, BuildSettings.ResizeDuringBuildY,
 				TargetTextureSizeX, TargetTextureSizeY, TargetTextureSizeZ);
@@ -4298,7 +4227,7 @@ private:
 				GenerateMipChain(BuildSettings, BaseMip, BuildSourceImageMips, 1);
 
 				// mip data not needed anymore
-				// todo: this could free "BaseMip" image instead, if it's ok with caller (currently const type used)
+				// @todo: this could free "BaseMip" image instead, if it's ok with caller (currently const type used)
 				Temp.RawData.Empty();
 
 				while( BuildSourceImageMips.Last().SizeX > MaxTextureResolution || 
@@ -4454,7 +4383,7 @@ private:
 						{
 							UE_LOG(LogTextureCompressor, Display, TEXT("[alpha] Copy: %d - %.*s"), FImageCore::DetectAlphaChannel(Mip), DebugTexturePathName.Len(), DebugTexturePathName.GetData());
 						}
-						//@@CB todo : when Mip format == Image format, we can Move instead of Copy
+						//@todo : when Mip format == Image format, we can Move instead of Copy
 						//	have to make sure that's okay with SourceMips/TextureData
 					}
 				}
diff --git a/Engine/Source/Runtime/Engine/Private/TextureDerivedDataTask.cpp b/Engine/Source/Runtime/Engine/Private/TextureDerivedDataTask.cpp
--- a/Engine/Source/Runtime/Engine/Private/TextureDerivedDataTask.cpp
+++ b/Engine/Source/Runtime/Engine/Private/TextureDerivedDataTask.cpp
@@ -1007,7 +1007,17 @@ static int64 GetBuildRequiredMemoryEstimate(UTexture* InTexture,
 				}
 			}
 		
-			int64 CurrentBlockTopMipNumPixels = (int64)SourceBlock.SizeX * SourceBlock.SizeY * SourceBlock.NumSlices;;
+			// assume pow2 options are the same for all layers, just use layer 0 here :
+			const FTextureBuildSettings & BuildSettings = InSettingsPerLayerFetchFirst[0];
+
+			int32 TargetSizeX, TargetSizeY, TargetSizeZ;
+			UE::TextureBuildUtilities::GetPowerOfTwoTargetTextureSize(SourceBlock.SizeX,SourceBlock.SizeY,SourceBlock.NumSlices,
+				BuildSettings.bVolume, (ETexturePowerOfTwoSetting::Type)BuildSettings.PowerOfTwoMode, 
+				BuildSettings.ResizeDuringBuildX, BuildSettings.ResizeDuringBuildY, 
+				TargetSizeX, TargetSizeY, TargetSizeZ);
+
+			int64 CurrentBlockTopMipNumPixels = (int64)TargetSizeX * TargetSizeY * TargetSizeZ;
+
 			TotalTopMipNumPixelsPerLayer += CurrentBlockTopMipNumPixels;
 
 			LargestBlockTopMipNumPixels = FMath::Max( CurrentBlockTopMipNumPixels , LargestBlockTopMipNumPixels );
@@ -1135,7 +1145,6 @@ static int64 GetBuildRequiredMemoryEstimate(UTexture* InTexture,
 
 		// Compute the memory it should take to uncompress the bulkdata in memory
 		int64 TotalSourceBytes = 0;
-		int64 TotalTopMipNumPixels = 0;
 
 		FTextureSourceBlock SourceBlock;
 		Source.GetBlock(0, SourceBlock);
@@ -1145,12 +1154,20 @@ static int64 GetBuildRequiredMemoryEstimate(UTexture* InTexture,
 			TotalSourceBytes += Source.CalcMipSize(0, 0, MipIndex);
 		}
 		
-		TotalTopMipNumPixels += (int64)SourceBlock.SizeX * SourceBlock.SizeY * SourceBlock.NumSlices;
-	
 		if ( TotalSourceBytes <= 0 )
 		{
 			return -1; /* Unknown */
 		}
+		
+		const FTextureBuildSettings & BuildSettings = InSettingsPerLayerFetchFirst[0];
+		
+		int32 TargetSizeX, TargetSizeY, TargetSizeZ;
+		UE::TextureBuildUtilities::GetPowerOfTwoTargetTextureSize(SourceBlock.SizeX,SourceBlock.SizeY,SourceBlock.NumSlices,
+			BuildSettings.bVolume, (ETexturePowerOfTwoSetting::Type)BuildSettings.PowerOfTwoMode, 
+			BuildSettings.ResizeDuringBuildX, BuildSettings.ResizeDuringBuildY, 
+			TargetSizeX, TargetSizeY, TargetSizeZ);
+
+		int64 TotalTopMipNumPixels = (int64)TargetSizeX * TargetSizeY * TargetSizeZ;
 
 		// assume full mip chain :
 		int64 TotalNumPixels = (TotalTopMipNumPixels * 4)/3;
@@ -1161,8 +1178,6 @@ static int64 GetBuildRequiredMemoryEstimate(UTexture* InTexture,
 		
 		int64 MemoryEstimate = TotalSourceBytes + IntermediateFloatColorBytes;
 	
-		const FTextureBuildSettings & BuildSettings = InSettingsPerLayerFetchFirst[0];
-
 		EPixelFormat PixelFormat = GetOutputPixelFormat(BuildSettings,bHasAlpha);
 
 		if ( PixelFormat == PF_Unknown )
@@ -1205,6 +1220,7 @@ static int64 GetBuildRequiredMemoryEstimate(UTexture* InTexture,
 			}
 			else if ( TextureFormatName == "BC4" || TextureFormatName == "BC5" )
 			{
+				// changed: TFO uses 2_U16 now (4 byte intermediate)
 				IntermediateBytesPerPixel = 8; // RGBA16
 			}
 			else
@@ -1220,6 +1236,7 @@ static int64 GetBuildRequiredMemoryEstimate(UTexture* InTexture,
 			if ( bRDO )
 			{
 				// activity map for whole image :
+				// (this has changed in newer versions of Oodle Texture)
 
 				// Phase1 = computing activity map
 				int ActivityBytesPerPixel;
@@ -1260,10 +1277,13 @@ static int64 GetBuildRequiredMemoryEstimate(UTexture* InTexture,
 		{
 			// ASTCenc does an entermediate copy to RGBA16F for HDR formats and RGBA8 for LDR
 			MemoryEstimate += (IsHDR(PixelFormat) ? 8 : 4) * TotalNumPixels;
+			// internal memory use of ASTCenc is not estimated
+			// @todo : fix me
 		}
 		else
 		{
 			// note: memory ues of non-Oodle encoders is not estimated
+			// @todo : fix me
 		}
 		
 		MemoryEstimate += 64 * 1024; // overhead room
@@ -1273,6 +1293,7 @@ static int64 GetBuildRequiredMemoryEstimate(UTexture* InTexture,
 		return MemoryEstimate;
 
 		// @todo Oodle : not right for volumes & latlong cubes
+		// @todo Oodle : not right with Composite , CPU textures
 	}
 }
 
