TextureCompressorModule use ResizeImage for MaxTextureResolution so we can resize directly from source to dest size
bump DDC key for change in output of MaxTextureResolution affected textures
Also free source images and temporaries as they are consumed and no longer needed to reduce peak mem use
In the nop case, move the image instead of copying

diff --git a/Engine/Source/Developer/TextureCompressor/Private/TextureCompressorModule.cpp b/Engine/Source/Developer/TextureCompressor/Private/TextureCompressorModule.cpp
--- a/Engine/Source/Developer/TextureCompressor/Private/TextureCompressorModule.cpp
+++ b/Engine/Source/Developer/TextureCompressor/Private/TextureCompressorModule.cpp
@@ -26,6 +26,12 @@
 #include "Windows/WindowsHWrapper.h"
 #endif
 
+#if 0
+// for saving image dumps:
+#include "IImageWrapperModule.h"
+#include "Misc/FileHelper.h"
+#endif
+
 static TAutoConsoleVariable<int32> CVarDetailedMipAlphaLogging(
 	TEXT("r.DetailedMipAlphaLogging"),
 	0,
@@ -2969,6 +2975,8 @@ void ITextureCompressorModule::AdjustImageColors(FImage& Image, const FTextureBu
 	}
 }
 
+// if alpha cannot be determined, returns false and does not fill bOutAlphaIsTransparent
+// if possible, returns true, fills out bOutAlphaIsTransparent
 bool ITextureCompressorModule::DetermineAlphaChannelTransparency(const FTextureBuildSettings& InBuildSettings, const FLinearColor& InChannelMin, const FLinearColor& InChannelMax, bool& bOutAlphaIsTransparent)
 {
 	// Settings that affect alpha:
@@ -3042,13 +3050,21 @@ bool ITextureCompressorModule::DetermineAlphaChannelTransparency(const FTextureB
 		ChannelMinMax[0].A = ChannelMinMax[0].R;
 		ChannelMinMax[1].A = ChannelMinMax[1].R;
 	}
+	
+	// use the same test as FImageCore::DetectAlphaChannel
+	// @todo : this is NOT the same threshold used for RGBA16 in DetectAlphaChannel
+	const float FloatNonOpaqueAlpha = 254.5f / 255.f; // the U8 alpha threshold
 
-	bOutAlphaIsTransparent = false;
-	if (ChannelMinMax[0].A < 1 ||
-		ChannelMinMax[1].A < 1)
+	float MinAlpha = FMath::Min(ChannelMinMax[0].A ,ChannelMinMax[1].A);
+
+	bOutAlphaIsTransparent = ( MinAlpha <= FloatNonOpaqueAlpha );
+	
+	// detect ambiguity zone where alpha could easily cross the threshold due to float math :
+	if ( MinAlpha >= (254.4f/255.f) && MinAlpha <= (254.6f/255.f) )
 	{
-		bOutAlphaIsTransparent = true;
+		return false; // not possible to answer bOutAlphaIsTransparent
 	}
+
 	return true;
 }
 
@@ -3760,9 +3776,10 @@ public:
 		return MipHashBuilder.Finalize().Hash;
 	}
 
+	// SourceMips can be freed by this call
 	virtual bool BuildTexture(
-		const TArray<FImage>& SourceMips,
-		const TArray<FImage>& AssociatedNormalSourceMips,
+		TArray<FImage>& SourceMips,
+		TArray<FImage>& AssociatedNormalSourceMips,
 		const FTextureBuildSettings& BuildSettings,
 		FStringView DebugTexturePathName,
 		TArray<FCompressedImage2D>& OutTextureMips,
@@ -3814,6 +3831,8 @@ public:
 			return false;
 		}
 		
+		SourceMips.Empty();
+
 		// apply roughness adjustment depending on normal map variation
 		if (AssociatedNormalSourceMips.Num())
 		{
@@ -3852,6 +3871,8 @@ public:
 				return false;
 			}
 
+			AssociatedNormalSourceMips.Empty();
+
 			if (!ApplyCompositeTextureToMips(IntermediateMipChain, IntermediateAssociatedNormalSourceMipChain, BuildSettings.CompositeTextureMode, BuildSettings.CompositePower, BuildSettings.LODBias))
 			{
 				UE_LOG(LogTextureCompressor, Display, TEXT("ApplyCompositeTextureToMips failed [%.*s]"),
@@ -3980,8 +4001,9 @@ public:
 private:
 
 
+	// InSourceMipChain can be freed by this call
 	bool BuildTextureMips(
-		const TArray<FImage>& InSourceMipChain,
+		TArray<FImage>& InSourceMipChain,
 		const FTextureBuildSettings& BuildSettings,
 		const bool bNeedLinearize,
 		TArray<FImage>& OutMipChain,
@@ -3991,6 +4013,14 @@ private:
 		
 		check(InSourceMipChain.Num() > 0);
 		check(InSourceMipChain[0].SizeX > 0 && InSourceMipChain[0].SizeY > 0 && InSourceMipChain[0].NumSlices > 0);
+		
+		int32 InSourceMipChainNum = InSourceMipChain.Num();
+
+		// Calculation check before we change anything, because source mips can be freed :
+		int32 CalculatedMip0SizeX, CalculatedMip0SizeY, CalculatedMip0NumSlices;
+		int32 CalculatedMipCount = GetMipCountForBuildSettings(
+			InSourceMipChain[0].SizeX, InSourceMipChain[0].SizeY, InSourceMipChain[0].NumSlices, InSourceMipChain.Num(), BuildSettings,
+			CalculatedMip0SizeX, CalculatedMip0SizeY, CalculatedMip0NumSlices);
 
 		// Identify long-lat cubemaps.
 		const bool bLongLatCubemap = BuildSettings.bLongLatSource;
@@ -4016,7 +4046,8 @@ private:
 		//	but that will change output :(
 
 		// pSourceMips will track the current FImages we consider to be "source"
-		const TArray<FImage> * pSourceMips = &InSourceMipChain;
+		//	whenever pSourceMips is changed, previous pSourceMips is freed
+		TArray<FImage> * pSourceMips = &InSourceMipChain;
 
 		// first pad up to pow2 if requested
 		ETexturePowerOfTwoSetting::Type PowerOfTwoMode = (ETexturePowerOfTwoSetting::Type) BuildSettings.PowerOfTwoMode;
@@ -4157,6 +4188,7 @@ private:
 					}
 				}
 				// change pSourceMips to point at the one padded image we made
+				pSourceMips->Empty();
 				pSourceMips = &PaddedSourceMips;
 			}
 		}		
@@ -4196,19 +4228,11 @@ private:
 				// the source is larger than the compressor allows and no mip image exists to act as a smaller source.
 				// We must generate a suitable source image:
 				const FImage& BaseImage = pSourceMips->Last();
-				bool bSuitableFormat = BaseImage.Format == ERawImageFormat::RGBA32F;
 			
 				check( MaxTextureResolution > 0 );
 				check( BaseImage.SizeX > MaxTextureResolution || 
 					   BaseImage.SizeY > MaxTextureResolution );
-
-				FImage Temp;
-				if (!bSuitableFormat)
-				{
-					// convert to RGBA32F
-					BaseImage.CopyTo(Temp, ERawImageFormat::RGBA32F, EGammaSpace::Linear);
-				}
-
+					   
 				UE_LOG(LogTextureCompressor, Verbose,
 					TEXT("Source image %dx%d too large for compressors max dimension (%d). Resizing."),
 					BaseImage.SizeX,
@@ -4216,39 +4240,113 @@ private:
 					BuildSettings.MaxTextureResolution
 					);
 
-				// make sure BuildSourceImageMips doesn't reallocate :
-				constexpr int BuildSourceImageMipsMaxCount = 20; // plenty
-				BuildSourceImageMips.Empty(BuildSourceImageMipsMaxCount);
+				if ( BuildSettings.bUseNewMipFilter && 
+					! BuildSettings.bVolume && ! BuildSettings.bDoScaleMipsForAlphaCoverage && ! BuildSettings.bPreserveBorder )
+				{
+					// BuildSourceImageMips will be used as the source for further steps, fill it :
+					BuildSourceImageMips.SetNum(1);
+					FImage & DestImage = BuildSourceImageMips[0];
+
+					// note we could resize directly to MaxTextureResolution
+					//	for now we will replicate the old logic of only doing 2X mip size steps
+					//	this results in output that is in (MaxTextureResolution/2,MaxTextureResolution]
+					int64 DestSizeX = BaseImage.SizeX;
+					int64 DestSizeY = BaseImage.SizeY;
+					while( DestSizeX > MaxTextureResolution || DestSizeY > MaxTextureResolution )
+					{
+						DestSizeX = FMath::Max(DestSizeX>>1,1);
+						DestSizeY = FMath::Max(DestSizeY>>1,1);
+					}
+
+					DestImage.Init(DestSizeX,DestSizeY,ERawImageFormat::RGBA32F, EGammaSpace::Linear);
+
+					// rather than doing a bunch of mip-step resizes,
+					//	just resize once directly to dest size using ResizeImage
+					//	also the BaseImage stays in BGRA8, no need to promote it to RGBA32F
+
+					// ResizeImage does work with Slices for cubes/arrays , but not volumes
+					check( ! BuildSettings.bVolume );
+
+					// @todo : find closest ResizeFilter that matches MipGen settings (maybe?)
+					//	 I mean, maybe not?  Most of the time MipGen is "SimpleAverage" which produces a really bad resize
+					FImageCore::EResizeImageFilter ResizeFilter = FImageCore::EResizeImageFilter::Default;
+					FImageCore::ResizeImage(BaseImage,DestImage,ResizeFilter);
 
-				// Max Texture Size resizing happens here :
-				// note we do not check for TMGS_Angular here
-				// note that TMGS_NoMipMaps *can* use this path; in that case it generates mips using 2x2 simple average
-				const FImage& BaseMip = bSuitableFormat ? BaseImage : Temp;
-				GenerateMipChain(BuildSettings, BaseMip, BuildSourceImageMips, 1);
+					#if 0
+					{
+						UE_LOG(LogTextureCompressor, Display, TEXT("ResizeImage from: %d x %d x %s -> %d x %d x %s"),
+							BaseImage.SizeX,BaseImage.SizeY,ERawImageFormat::GetName(BaseImage.Format),
+							DestImage.SizeX,DestImage.SizeY,ERawImageFormat::GetName(DestImage.Format));
+						
+						IImageWrapperModule& ImageWrapperModule = FModuleManager::LoadModuleChecked<IImageWrapperModule>( FName("ImageWrapper") );
+		
+						TArray64<uint8> Buffer;
+						ImageWrapperModule.CompressImage(Buffer,EImageFormat::PNG,BaseImage,0);
+						FFileHelper::SaveArrayToFile(Buffer,TEXT("r:\\BaseImage.png"));
+						
+						ImageWrapperModule.CompressImage(Buffer,EImageFormat::EXR,DestImage,0);
+						FFileHelper::SaveArrayToFile(Buffer,TEXT("r:\\DestImage.exr"));
 
-				// mip data not needed anymore
-				// @todo: this could free "BaseMip" image instead, if it's ok with caller (currently const type used)
-				Temp.RawData.Empty();
+						FImage DestImage8;
+						DestImage.CopyTo(DestImage8,ERawImageFormat::BGRA8,EGammaSpace::sRGB);
+						ImageWrapperModule.CompressImage(Buffer,EImageFormat::PNG,DestImage8,0);
+						FFileHelper::SaveArrayToFile(Buffer,TEXT("r:\\DestImage.png"));
+						
+						bool BaseAlpha = FImageCore::DetectAlphaChannel(BaseImage);
+						bool DestAlpha = FImageCore::DetectAlphaChannel(DestImage);
 
-				while( BuildSourceImageMips.Last().SizeX > MaxTextureResolution || 
-					   BuildSourceImageMips.Last().SizeY > MaxTextureResolution )
+						check( BaseAlpha == DestAlpha );
+					}
+					#endif
+				}
+				else
 				{
-					// note: now making mips one by one, rather than N in one call
-					//	this is not exactly the same if AlphaCoverage processing is on
-					check( BuildSourceImageMips.Num() < BuildSourceImageMipsMaxCount );
-					FImage& LastMip = BuildSourceImageMips.Last();
-					GenerateMipChain(BuildSettings, LastMip, BuildSourceImageMips, 1);
+					// old way
+
+					FImage Temp;
+					bool bSuitableFormat = BaseImage.Format == ERawImageFormat::RGBA32F;
+					if (!bSuitableFormat)
+					{
+						// convert to RGBA32F
+						BaseImage.CopyTo(Temp, ERawImageFormat::RGBA32F, EGammaSpace::Linear);
+					}
+
+					// make sure BuildSourceImageMips doesn't reallocate :
+					constexpr int BuildSourceImageMipsMaxCount = 20; // plenty
+					BuildSourceImageMips.Empty(BuildSourceImageMipsMaxCount);
+
+					// Max Texture Size resizing happens here :
+					// note we do not check for TMGS_Angular here
+					// note that TMGS_NoMipMaps *can* use this path; in that case it generates mips using 2x2 simple average
+					const FImage& BaseMip = bSuitableFormat ? BaseImage : Temp;
+					GenerateMipChain(BuildSettings, BaseMip, BuildSourceImageMips, 1);
 
 					// mip data not needed anymore
-					LastMip.RawData.Empty();
-				}
-			
+					Temp.RawData.Empty();					
+					pSourceMips->Empty();
+
+					while( BuildSourceImageMips.Last().SizeX > MaxTextureResolution || 
+						   BuildSourceImageMips.Last().SizeY > MaxTextureResolution )
+					{
+						// note: now making mips one by one, rather than N in one call
+						//	this is not exactly the same if AlphaCoverage processing is on
+						check( BuildSourceImageMips.Num() < BuildSourceImageMipsMaxCount );
+						FImage& LastMip = BuildSourceImageMips.Last();
+						GenerateMipChain(BuildSettings, LastMip, BuildSourceImageMips, 1);
+
+						// mip data not needed anymore
+						LastMip.RawData.Empty();
+					}
+				}				
+
 				check( BuildSourceImageMips.Last().SizeX <= MaxTextureResolution &&
-					   BuildSourceImageMips.Last().SizeY <= MaxTextureResolution );
+						BuildSourceImageMips.Last().SizeY <= MaxTextureResolution );
 
-				// change pSourceMips to point at the mip chain we made
+				// change pSourceMips to point at the mip chain we made				
+				pSourceMips->Empty();
 				pSourceMips = &BuildSourceImageMips;
 				StartMip = BuildSourceImageMips.Num() - 1;
+
 				// [StartMip] will now references BuildSourceImageMips.Last()
 			}
 		}
@@ -4284,7 +4382,7 @@ private:
 			// LeaveExistingMips is often inadvertently used as "NoMips", so if they only brought in 1 mip, leave it as
 			// 1 mip.
 			if (BuildSettings.MipGenSettings == TMGS_LeaveExistingMips &&
-				InSourceMipChain.Num() == 1)
+				InSourceMipChainNum == 1)
 			{
 				NumOutputMips = 1;
 				check(CopyCount == 1);
@@ -4309,11 +4407,11 @@ private:
 		const bool bLinearize = bNeedLinearize || (GenerateCount > 0) || BuildSettings.bRenormalizeTopMip || (BuildSettings.Downscale > 1.f)
 			|| BuildSettings.bHasColorSpaceDefinition || BuildSettings.bComputeBokehAlpha || BuildSettings.bFlipGreenChannel
 			|| BuildSettings.bReplicateRed || BuildSettings.bReplicateAlpha || BuildSettings.bApplyYCoCgBlockScale
-			|| BuildSettings.SourceEncodingOverride != 0 || bNeedAdjustImageColors || BuildSettings.bNormalizeNormals;
+			|| (BuildSettings.SourceEncodingOverride != 0) || bNeedAdjustImageColors || BuildSettings.bNormalizeNormals;
 
 		for (int32 MipIndex = StartMip; MipIndex < StartMip + CopyCount; ++MipIndex)
 		{
-			const FImage& Image = (*pSourceMips)[MipIndex];
+			FImage& Image = (*pSourceMips)[MipIndex];
 
 			if (bDoDetailedAlphaLogging)
 			{
@@ -4370,26 +4468,26 @@ private:
 					{
 						// if image is in BGRA8 format leave it, otherwise convert to RGBA32F
 						//  we only support leaving images in source format if they are BGRA8 and require no processing (eg VT tiles)
-						ERawImageFormat::Type DestFormat = Image.Format;
-						EGammaSpace DestGammaSpace = Image.GammaSpace;
-						if ( Image.Format != ERawImageFormat::BGRA8 )
+						if ( Image.Format == ERawImageFormat::BGRA8 )
 						{
-							DestFormat = ERawImageFormat::RGBA32F;
-							DestGammaSpace = EGammaSpace::Linear;
+							// just move Image to Mip, no copy
+							Image.Swap(Mip);
 						}
-						Image.CopyTo(Mip, DestFormat, DestGammaSpace);
-						
+						else
+						{
+							Image.CopyTo(Mip, ERawImageFormat::RGBA32F, EGammaSpace::Linear);
+						}
+
 						if (bDoDetailedAlphaLogging)
 						{
 							UE_LOG(LogTextureCompressor, Display, TEXT("[alpha] Copy: %d - %.*s"), FImageCore::DetectAlphaChannel(Mip), DebugTexturePathName.Len(), DebugTexturePathName.GetData());
 						}
-						//@todo : when Mip format == Image format, we can Move instead of Copy
-						//	have to make sure that's okay with SourceMips/TextureData
 					}
 				}
 			}
 
-			//@todo Oodle: free SourceMips as we consume them?
+			// free pSourceMips as we consume them
+			Image.RawData.Empty();
 
 			if (BuildSettings.Downscale > 1.f)
 			{		
@@ -4428,6 +4526,10 @@ private:
 		check( OutMipChain.Num() == CopyCount );
 		check( GenerateCount == NumOutputMips - OutMipChain.Num() );
 
+		// no further reads from source:
+		pSourceMips->Empty();
+		pSourceMips = nullptr;
+
 		// Generate any missing mips in the chain.
 		if ( GenerateCount > 0 )
 		{
@@ -4459,10 +4561,6 @@ private:
 		}
 		check(OutMipChain.Num() == NumOutputMips);
 
-		int32 CalculatedMip0SizeX, CalculatedMip0SizeY, CalculatedMip0NumSlices;
-		int32 CalculatedMipCount = GetMipCountForBuildSettings(
-			InSourceMipChain[0].SizeX, InSourceMipChain[0].SizeY, InSourceMipChain[0].NumSlices, InSourceMipChain.Num(), BuildSettings,
-			CalculatedMip0SizeX, CalculatedMip0SizeY, CalculatedMip0NumSlices);
 		if (CalculatedMipCount != NumOutputMips ||
 			CalculatedMip0SizeX != OutMipChain[0].SizeX ||
 			CalculatedMip0SizeY != OutMipChain[0].SizeY ||
diff --git a/Engine/Source/Developer/TextureCompressor/Public/TextureCompressorModule.h b/Engine/Source/Developer/TextureCompressor/Public/TextureCompressorModule.h
--- a/Engine/Source/Developer/TextureCompressor/Public/TextureCompressorModule.h
+++ b/Engine/Source/Developer/TextureCompressor/Public/TextureCompressorModule.h
@@ -393,10 +393,12 @@ public:
 	 * @param OutNumMipsInTail - The number of mips that are joined into a single mip tail mip
 	 * @param OutExtData - Extra data that the runtime may need
 	 * @returns true on success
+	 
+	// SourceMips can be freed by this call
 	 */
 	virtual bool BuildTexture(
-		const TArray<struct FImage>& SourceMips,
-		const TArray<struct FImage>& AssociatedNormalSourceMips,
+		TArray<struct FImage>& SourceMips,
+		TArray<struct FImage>& AssociatedNormalSourceMips,
 		const FTextureBuildSettings& BuildSettings,
 		FStringView DebugTexturePathName,
 		TArray<FCompressedImage2D>& OutTextureMips,
diff --git a/Engine/Source/Runtime/Engine/Private/TextureDerivedData.cpp b/Engine/Source/Runtime/Engine/Private/TextureDerivedData.cpp
--- a/Engine/Source/Runtime/Engine/Private/TextureDerivedData.cpp
+++ b/Engine/Source/Runtime/Engine/Private/TextureDerivedData.cpp
@@ -302,6 +302,14 @@ static void SerializeForKey(FArchive& Ar, const FTextureBuildSettings& Settings)
 		TempGuid = FGuid(0x418B8584, 0x72D54EA5, 0xBA8E8C2B, 0xECC880DE);
 		Ar << TempGuid;
 	}
+	
+	if ( Settings.MaxTextureResolution != FTextureBuildSettings::MaxTextureResolutionDefault && Settings.bUseNewMipFilter )
+	{
+		// @todo SerializeForKey these can go away whenever we bump the overall ddc key
+		// behavior of MaxTextureResolution changed to ResizeImage 2/8/2024
+		TempName = TEXTURE_DDC_STB_IMAGE_RESIZE_VERSION;
+		Ar << TempName;
+	}
 
 	if ( Settings.bVolume )
 	{
@@ -4014,9 +4022,6 @@ bool UTexture::DownsizeImageUsingTextureSettings(const ITargetPlatform* TargetPl
 	constexpr int BuildSourceImageMipsMaxCount = 20; // plenty
 	BuildSourceImageMips.Empty(BuildSourceImageMipsMaxCount);
 
-	// one nice thing we do get from GenerateMipChain (as opposed to ResizeImage)
-	//	is that wrap/clamp address mode is respected and cubemaps clamp
-
 	ITextureCompressorModule::GenerateMipChain(BuildSettings, Temp, BuildSourceImageMips, 1);
 
 	// keep halving while larger size is > Target
diff --git a/Engine/Source/Runtime/Engine/Private/TextureDerivedDataTask.cpp b/Engine/Source/Runtime/Engine/Private/TextureDerivedDataTask.cpp
--- a/Engine/Source/Runtime/Engine/Private/TextureDerivedDataTask.cpp
+++ b/Engine/Source/Runtime/Engine/Private/TextureDerivedDataTask.cpp
@@ -831,25 +831,36 @@ static void DDC1_BuildTexture(
 
 		uint32 NumMipsInTail;
 		uint32 ExtData;
+		
+		TArray<FImage> EmptyImageArray;
+		TArray<FImage> & CompositeImageArray = ((bool)Texture.GetCompositeTexture() && CompositeTextureData.Blocks.Num() && CompositeTextureData.Blocks[0].MipsPerLayer.Num()) ? CompositeTextureData.Blocks[0].MipsPerLayer[0] : EmptyImageArray;
 
 		// Compress the texture by calling texture compressor directly.
 		TArray<FCompressedImage2D> CompressedMips;
-		if (Compressor->BuildTexture(TextureData.Blocks[0].MipsPerLayer[0],
-			((bool)Texture.GetCompositeTexture() && CompositeTextureData.Blocks.Num() && CompositeTextureData.Blocks[0].MipsPerLayer.Num()) ? CompositeTextureData.Blocks[0].MipsPerLayer[0] : TArray<FImage>(),
+
+		bSucceeded = Compressor->BuildTexture(TextureData.Blocks[0].MipsPerLayer[0],
+			CompositeImageArray,
 			InBuildSettingsPerLayer[0],
 			TexturePathName,
 			CompressedMips,
 			NumMipsInTail,
 			ExtData,
-			nullptr))
+			nullptr);
+			
+
+		if ( bSucceeded )
 		{
+			// BuildTexture can free the source images passed to it
+			//	so TextureData is invalid after this call
+			TextureData.ReleaseMemory();
+			CompositeTextureData.ReleaseMemory();;
+
 			check(CompressedMips.Num());
 
 			DDC1_StoreClassicTextureInDerivedData(
 				CompressedMips, DerivedData, InBuildSettingsPerLayer[0].bVolume, InBuildSettingsPerLayer[0].bTextureArray, InBuildSettingsPerLayer[0].bCubemap, 
 				NumMipsInTail, ExtData, bReplaceExistingDDC, TexturePathName, KeySuffix, BytesCached);
 
-			bSucceeded = true;
 			DerivedData->ResultMetadata = InBuildResultMetadata;
 
 			const bool bInlineMips = EnumHasAnyFlags(CacheFlags, ETextureCacheFlags::InlineMips);
@@ -870,7 +881,6 @@ static void DDC1_BuildTexture(
 			// will log below
 			check( DerivedData->Mips.Num() == 0 );
 			DerivedData->Mips.Empty();
-			bSucceeded = false;
 
 			UE_LOG(LogTexture, Warning, TEXT("BuildTexture failed to build %s derived data for %s"), *InBuildSettingsPerLayer[0].TextureFormatName.GetPlainNameString(), *TexturePathName);
 		}
@@ -1861,6 +1871,10 @@ bool DDC1_BuildTiledClassicTexture(
 				LinearBytesCached,
 				bHasLinearDerivedData);
 
+			// TextureData can be freed by Build, don't use it anymore :
+			TextureData.ReleaseMemory();
+			CompositeTextureData.ReleaseMemory();
+
 			// This should succeed because we asked for inline mips if the build succeeded
 			if (bHasLinearDerivedData && 
 				LinearDerivedData.TryLoadMipsWithSizes(0, LinearMipData, LinearMipSizes, TexturePathName) == false)
@@ -1887,22 +1901,6 @@ bool DDC1_BuildTiledClassicTexture(
 	FEncodedTextureExtendedData TextureExtendedData;
 	int32 OutputTextureNumStreamingMips;
 	{
-#if DO_CHECK
-		TArray<FImage>& SourceMips = TextureData.Blocks[0].MipsPerLayer[0];
-		if (SourceMips.Num()) // if we're tiling existing linear data then the texture data isn't actually loaded so we can't do this sanity check.
-		{
-			int32 LinearMip0SizeX, LinearMip0SizeY, LinearMip0NumSlices;
-			int32 LinearMipCount = Compressor->GetMipCountForBuildSettings(SourceMips[0].SizeX, SourceMips[0].SizeY, SourceMips[0].NumSlices, SourceMips.Num(), LinearSettingsPerLayerFetchOrBuild[0],
-				LinearMip0SizeX,
-				LinearMip0SizeY,
-				LinearMip0NumSlices);
-			check(LinearDerivedData.Mips[0].SizeX == LinearMip0SizeX);
-			check(LinearDerivedData.Mips[0].SizeY == LinearMip0SizeY);
-			check(LinearDerivedData.Mips.Num() == LinearMipCount);
-			check(LinearDerivedData.GetNumSlices() == LinearMip0NumSlices);
-		}
-#endif
-
 		LinearSettingsPerLayerFetchOrBuild[0].GetEncodedTextureDescriptionWithPixelFormat(
 			&TextureDescription,
 			LinearDerivedData.PixelFormat, LinearDerivedData.Mips[0].SizeX, LinearDerivedData.Mips[0].SizeY, LinearDerivedData.GetNumSlices(), LinearDerivedData.Mips.Num());
@@ -2068,6 +2066,8 @@ void FTextureCacheDerivedDataWorker::DoWork()
 					KeySuffix, bReplaceExistingDDC, RequiredMemoryEstimate, DerivedData, BytesCached, bSucceeded);
 			}
 
+			// TextureData may have been freed by Build, don't use it anymore
+
 			if (bInvalidVirtualTextureCompression && DerivedData->VTData)
 			{
 				// If we loaded data that turned out to be corrupt, flag it here so we can also recreate the VT data cached to local /DerivedDataCache/VT/ directory
@@ -2133,7 +2133,7 @@ void FTextureCacheDerivedDataWorker::Finalize()
 	// --	always happens with a ForceRebuildPlatformData, which is called whenever mip data is requested
 	//		in the editor and is missing for some reason.
 	// --	always with a lighting build, as the async light/shadowmap tasks will disallow async builds. 
-	// --	if the texture compiler cvar disallows async texture compilation
+	// --	if the texture compiler cvar disallows async texture compilation  "Editor.AsyncTextureCompilation 0"
 
 	if ( bTriedAndFailed )
 	{
@@ -2149,6 +2149,7 @@ void FTextureCacheDerivedDataWorker::Finalize()
 			return;
 		}
 
+		// note: GetSourceMips will not even try if TextureData.bValid was set to false
 		TextureData.GetSourceMips(Texture.Source, ImageWrapper);
 		if (Texture.GetCompositeTexture() && Texture.GetCompositeTexture()->Source.IsValid())
 		{
diff --git a/Engine/Source/Runtime/Engine/Private/VT/VirtualTextureDataBuilder.cpp b/Engine/Source/Runtime/Engine/Private/VT/VirtualTextureDataBuilder.cpp
--- a/Engine/Source/Runtime/Engine/Private/VT/VirtualTextureDataBuilder.cpp
+++ b/Engine/Source/Runtime/Engine/Private/VT/VirtualTextureDataBuilder.cpp
@@ -754,8 +754,6 @@ void FVirtualTextureDataBuilder::BuildLayerBlocks(FSlowTask& BuildTask, uint32 L
 {
 	TRACE_CPUPROFILER_EVENT_SCOPE(Texture.VT.BuildLayerBlocks);
 
-	const TArray<FImage> EmptyImageArray;
-
 	const int32 TileSize = SettingsPerLayer[0].VirtualTextureTileSize;
 	const int32 NumLayers = SourceData.Layers.Num();
 	const int32 NumBlocks = SourceData.Blocks.Num();
@@ -792,7 +790,7 @@ void FVirtualTextureDataBuilder::BuildLayerBlocks(FSlowTask& BuildTask, uint32 L
 	{
 		BuildTask.EnterProgressFrame();
 
-		const FTextureSourceBlockData& SourceBlockData = SourceData.Blocks[BlockIndex];
+		FTextureSourceBlockData& SourceBlockData = SourceData.Blocks[BlockIndex];
 
 		// Current lock + mips that will be compressed to tiles
 		FVTBlockPayload& BlockData = LayerPayload[LayerIndex].Blocks[BlockIndex];
@@ -809,8 +807,9 @@ void FVirtualTextureDataBuilder::BuildLayerBlocks(FSlowTask& BuildTask, uint32 L
 
 		const FTextureBuildSettings& BuildSettingsForLayer = SettingsPerLayer[LayerIndex];
 
-		const TArray<FImage>& SourceMips = SourceBlockData.MipsPerLayer[LayerIndex];
-		const TArray<FImage>* CompositeSourceMips = &EmptyImageArray;
+		TArray<FImage>& SourceMips = SourceBlockData.MipsPerLayer[LayerIndex];
+		TArray<FImage> EmptyImageArray;
+		TArray<FImage>* CompositeSourceMips = &EmptyImageArray;
 		if (CompositeSourceData.Blocks.Num() > 0)
 		{
 			CompositeSourceMips = &CompositeSourceData.Blocks[BlockIndex].MipsPerLayer[LayerIndex];
@@ -875,6 +874,10 @@ void FVirtualTextureDataBuilder::BuildLayerBlocks(FSlowTask& BuildTask, uint32 L
 		{
 			uint32 NumMipsInTail, ExtData;
 			bBuildTextureResult = Compressor->BuildTexture(SourceMips, *CompositeSourceMips, TBSettings, CurDebugTexturePathName, CompressedMips, NumMipsInTail, ExtData, nullptr);
+
+			// BuildTexture can free mips, they are no longer valid
+			SourceMips.Empty();
+			CompositeSourceMips->Empty();
 		}
 		else
 		{
@@ -1042,11 +1045,14 @@ void FVirtualTextureDataBuilder::BuildLayerBlocks(FSlowTask& BuildTask, uint32 L
 		TArray<FCompressedImage2D> CompressedMips;
 		uint32 NumMipsInTail, ExtData;
 		// this is a Build to uncompressed, to apply processing
+		TArray<FImage> EmptyImageArray;
 		if (!Compressor->BuildTexture(MiptailInputImages, EmptyImageArray, TBSettings, CurDebugTexturePathName, CompressedMips, NumMipsInTail, ExtData, nullptr))
 		{
 			check(false);
 		}
 
+		MiptailInputImages.Empty();
+
 		// We skip the first compressed mip output, since that will just be a copy of the input
 		check(CompressedMips.Num() >= BlockData.NumMips + 1);
 		// not true with padding options :
diff --git a/Engine/Source/Runtime/ImageWrapper/Private/ImageWrapperModule.cpp b/Engine/Source/Runtime/ImageWrapper/Private/ImageWrapperModule.cpp
--- a/Engine/Source/Runtime/ImageWrapper/Private/ImageWrapperModule.cpp
+++ b/Engine/Source/Runtime/ImageWrapper/Private/ImageWrapperModule.cpp
@@ -386,6 +386,8 @@ public:
 	{
 		TRACE_CPUPROFILER_EVENT_SCOPE(ImageWrapper.Decompress);
 
+		OutImage = FImage();
+
 		EImageFormat ImageFormat = DetectImageFormat(InCompressedData,InCompressedSize);
 		if ( ImageFormat == EImageFormat::Invalid )
 		{
@@ -420,6 +422,8 @@ public:
 	{
 		TRACE_CPUPROFILER_EVENT_SCOPE(ImageWrapper.Compress);
 
+		OutData.Empty();
+
 		TSharedPtr<IImageWrapper> ImageWrapper = CreateImageWrapper(ToFormat);
 		if ( ! ImageWrapper.IsValid() )
 		{
