TextureCompressorModule : if CanAcceptNonF32Source allow all formats to pass through, not just BGRA8

diff --git a/Engine/Source/Developer/TextureCompressor/Private/TextureCompressorModule.cpp b/Engine/Source/Developer/TextureCompressor/Private/TextureCompressorModule.cpp
--- a/Engine/Source/Developer/TextureCompressor/Private/TextureCompressorModule.cpp
+++ b/Engine/Source/Developer/TextureCompressor/Private/TextureCompressorModule.cpp
@@ -1510,14 +1510,14 @@ static void DownscaleImage(const FImage& SrcImage, FImage& DstImage, const FText
 
 		if ( Settings.DownscaleOptions == (uint8)ETextureDownscaleOptions::Unfiltered )
 			Filter = FImageCore::EResizeImageFilter::PointSample;
-		else if ( Settings.DownscaleOptions == (uint8)ETextureDownscaleOptions::SimpleAverage )
+		else if ( Settings.DownscaleOptions == (uint8)ETextureDownscaleOptions::SimpleAverage ) // <- Default maps to SimpleAverage
 			Filter = FImageCore::EResizeImageFilter::Triangle;
 		else if ( Settings.DownscaleOptions >= (uint8)ETextureDownscaleOptions::Sharpen5 ) // sharpest
 			Filter = FImageCore::EResizeImageFilter::CubicSharp;
 		else if ( Settings.DownscaleOptions >= (uint8)ETextureDownscaleOptions::Sharpen0 ) // sharp
-			Filter = FImageCore::EResizeImageFilter::AdaptiveSharp;		
+			Filter = FImageCore::EResizeImageFilter::AdaptiveSharp;		 // winds up as CubicMitchell
 		else
-			Filter = FImageCore::EResizeImageFilter::AdaptiveSmooth;		
+			Filter = FImageCore::EResizeImageFilter::AdaptiveSmooth; // <- this can't be reached, Blur is not available as a DownscaleOptions
 
 		FImage Temp; // needed because Src == Dst
 		FImageCore::ResizeImageAllocDest(SrcImage,Temp,FinalSizeX,FinalSizeY,Filter);
@@ -4459,17 +4459,8 @@ private:
 					}
 					else
 					{
-						// if image is in BGRA8 format leave it, otherwise convert to RGBA32F
-						//  we only support leaving images in source format if they are BGRA8 and require no processing (eg VT tiles)
-						if ( Image.Format == ERawImageFormat::BGRA8 )
-						{
-							// just move Image to Mip, no copy
-							Image.Swap(Mip);
-						}
-						else
-						{
-							Image.CopyTo(Mip, ERawImageFormat::RGBA32F, EGammaSpace::Linear);
-						}
+						// just move Image to Mip, no copy
+						Image.Swap(Mip);
 
 						if (bDoDetailedAlphaLogging)
 						{
@@ -4479,11 +4470,8 @@ private:
 				}
 			}
 
-			{
-			TRACE_CPUPROFILER_EVENT_SCOPE(Texture.BuildTextureMips.Free);
-			// free pSourceMips as we consume them
+			// free pSourceMips as we consume them (detached)
 			Image.FreeData(true);
-			}
 
 			if (BuildSettings.Downscale > 1.f)
 			{		
