TextureCompressorModule : fix bug with MaxTextureSize on images with more than 1 slice (cubes, arrays)
bump DDC key for affected textures
bug introduced in CL 31470247

diff --git a/Engine/Source/Developer/TextureCompressor/Private/TextureCompressorModule.cpp b/Engine/Source/Developer/TextureCompressor/Private/TextureCompressorModule.cpp
--- a/Engine/Source/Developer/TextureCompressor/Private/TextureCompressorModule.cpp
+++ b/Engine/Source/Developer/TextureCompressor/Private/TextureCompressorModule.cpp
@@ -4309,7 +4309,7 @@ private:
 						DestSizeY = FMath::Max(DestSizeY>>1,1);
 					}
 
-					DestImage.Init(DestSizeX,DestSizeY,ERawImageFormat::RGBA32F, EGammaSpace::Linear);
+					DestImage.Init(DestSizeX,DestSizeY,BaseImage.NumSlices,ERawImageFormat::RGBA32F, EGammaSpace::Linear);
 
 					// rather than doing a bunch of mip-step resizes,
 					//	just resize once directly to dest size using ResizeImage
diff --git a/Engine/Source/Runtime/Engine/Private/TextureDerivedData.cpp b/Engine/Source/Runtime/Engine/Private/TextureDerivedData.cpp
--- a/Engine/Source/Runtime/Engine/Private/TextureDerivedData.cpp
+++ b/Engine/Source/Runtime/Engine/Private/TextureDerivedData.cpp
@@ -312,6 +312,12 @@ static void SerializeForKey(FArchive& Ar, const FTextureBuildSettings& Settings)
 		// behavior of MaxTextureResolution changed to ResizeImage 2/8/2024
 		TempName = TEXTURE_DDC_STB_IMAGE_RESIZE_VERSION;
 		Ar << TempName;
+
+		if ( Settings.bCubemap || Settings.bTextureArray )
+		{
+			TempName = FName(TEXT("Sliced Resize Bug Fix 03/07/2024"));
+			Ar << TempName;
+		}
 	}
 
 	if ( Settings.bVolume )
