TextureCompressor logging / behavior nop.  Move CompressedDataSize logging so it is done before the decode for platform preview.  Remove some disabled debug code.

diff --git a/Engine/Source/Developer/TextureCompressor/Private/TextureCompressorModule.cpp b/Engine/Source/Developer/TextureCompressor/Private/TextureCompressorModule.cpp
--- a/Engine/Source/Developer/TextureCompressor/Private/TextureCompressorModule.cpp
+++ b/Engine/Source/Developer/TextureCompressor/Private/TextureCompressorModule.cpp
@@ -4057,6 +4057,27 @@ public:
 			IntermediateMipChain.Empty();
 		}
 
+		if ( bCompressSucceeded )
+		{
+			// make sure this log is done before the decode for platform preview
+			//  -logcmds="LogTextureUpload Verbose,LogTextureCompressor Verbose"
+
+			for(int32 MipIndex=0;MipIndex < OutTextureMips.Num();MipIndex++)
+			{
+				const FCompressedImage2D & CompressedImage = OutTextureMips[MipIndex];
+
+				int64 CompressedDataSize = CompressedImage.RawData.Num();
+
+				// log what the TextureFormat built :
+				UE_LOG(LogTextureCompressor, Verbose, TEXT("Built texture: [%.*s] Compressed Mip %d PF=%d=%s : %dx%dx%d : CompressedDataSize=%lld"),
+					DebugTexturePathName.Len(),DebugTexturePathName.GetData(),
+					MipIndex, (int)CompressedImage.PixelFormat,
+					GetPixelFormatString((EPixelFormat)CompressedImage.PixelFormat),
+					CompressedImage.SizeX, CompressedImage.SizeY, CompressedImage.NumSlicesWithDepth,
+					CompressedDataSize);
+			}
+		}
+
 		return bCompressSucceeded;
 	}
 
@@ -4344,33 +4365,6 @@ private:
 					//	 I mean, maybe not?  Most of the time MipGen is "SimpleAverage" which produces a really bad resize
 					FImageCore::EResizeImageFilter ResizeFilter = FImageCore::EResizeImageFilter::Default;
 					FImageCore::ResizeImage(BaseImage,DestImage,ResizeFilter);
-
-					#if 0
-					{
-						UE_LOG(LogTextureCompressor, Display, TEXT("ResizeImage from: %d x %d x %s -> %d x %d x %s"),
-							BaseImage.SizeX,BaseImage.SizeY,ERawImageFormat::GetName(BaseImage.Format),
-							DestImage.SizeX,DestImage.SizeY,ERawImageFormat::GetName(DestImage.Format));
-						
-						IImageWrapperModule& ImageWrapperModule = FModuleManager::LoadModuleChecked<IImageWrapperModule>( FName("ImageWrapper") );
-		
-						TArray64<uint8> Buffer;
-						ImageWrapperModule.CompressImage(Buffer,EImageFormat::PNG,BaseImage,0);
-						FFileHelper::SaveArrayToFile(Buffer,TEXT("r:\\BaseImage.png"));
-						
-						ImageWrapperModule.CompressImage(Buffer,EImageFormat::EXR,DestImage,0);
-						FFileHelper::SaveArrayToFile(Buffer,TEXT("r:\\DestImage.exr"));
-
-						FImage DestImage8;
-						DestImage.CopyTo(DestImage8,ERawImageFormat::BGRA8,EGammaSpace::sRGB);
-						ImageWrapperModule.CompressImage(Buffer,EImageFormat::PNG,DestImage8,0);
-						FFileHelper::SaveArrayToFile(Buffer,TEXT("r:\\DestImage.png"));
-						
-						bool BaseAlpha = FImageCore::DetectAlphaChannel(BaseImage);
-						bool DestAlpha = FImageCore::DetectAlphaChannel(DestImage);
-
-						check( BaseAlpha == DestAlpha );
-					}
-					#endif
 				}
 				else
 				{
diff --git a/Engine/Source/Runtime/Engine/Private/TextureDerivedDataTask.cpp b/Engine/Source/Runtime/Engine/Private/TextureDerivedDataTask.cpp
--- a/Engine/Source/Runtime/Engine/Private/TextureDerivedDataTask.cpp
+++ b/Engine/Source/Runtime/Engine/Private/TextureDerivedDataTask.cpp
@@ -815,20 +815,6 @@ static void DDC1_StoreClassicTextureInDerivedData(
 		// CompressedDataSize can exceed int32 ; eg. 16k x 16k x RGBA16F == 2 GB
 		// DDC1 should be 64-bit safe now
 
-		// CompressedImage sizes were padded up to multiple of 4 for d3d, no longer ; log the align-up-to-4 sizes for debugging?
-
-		// log what the TextureFormat built :
-		//	(todo: change to "LogTexture" instead of "LogTextureUpload" and remove the align-up-to-4 debug logs)
-		UE_LOG(LogTextureUpload, Verbose, TEXT("Built texture: %s Compressed Mip %d PF=%d=%s : %dx%dx%d : %lld ; up4 %dx%d=%d"),
-			*TexturePathName,
-			MipIndex, (int)CompressedImage.PixelFormat,
-			GetPixelFormatString((EPixelFormat)CompressedImage.PixelFormat),
-			CompressedImage.SizeX, CompressedImage.SizeY, CompressedImage.NumSlicesWithDepth,
-			CompressedDataSize,
-			(CompressedImage.SizeX + 3) & (~3),
-			(CompressedImage.SizeY + 3) & (~3),
-			((CompressedImage.SizeX + 3) & (~3)) * ((CompressedImage.SizeY + 3) & (~3)));
-
 		NewMip->BulkData.Lock(LOCK_READ_WRITE);
 		void* NewMipData = NewMip->BulkData.Realloc(CompressedDataSize);
 		FMemory::Memcpy(NewMipData, CompressedImage.RawData.GetData(), CompressedDataSize);
