Switch to using source mips for alpha detection combined with prediction based on the build settings.

diff --git a/Engine/Source/Developer/TextureCompressor/Private/TextureCompressorModule.cpp b/Engine/Source/Developer/TextureCompressor/Private/TextureCompressorModule.cpp
--- a/Engine/Source/Developer/TextureCompressor/Private/TextureCompressorModule.cpp
+++ b/Engine/Source/Developer/TextureCompressor/Private/TextureCompressorModule.cpp
@@ -3097,10 +3097,12 @@ bool ITextureCompressorModule::DetermineAlphaChannelTransparency(const FTextureB
 	bOutAlphaIsTransparent = ( MinAlpha <= FloatNonOpaqueAlpha );
 	
 	// detect ambiguity zone where alpha could easily cross the threshold due to float math :
-	if ( MinAlpha >= (254.4f/255.f) && MinAlpha <= (254.6f/255.f) )
-	{
-		return false; // not possible to answer bOutAlphaIsTransparent
-	}
+	// Removed -- Any time this happens the author can't have predicted what alpha would be so
+	// use whatever result we select to be consistent. This is not hit in any known texture.
+	//if ( MinAlpha >= (254.4f/255.f) && MinAlpha <= (254.6f/255.f) )
+	//{
+	//	return false; // not possible to answer bOutAlphaIsTransparent
+	//}
 
 	return true;
 }
@@ -3480,7 +3482,7 @@ static bool CompressMipChain(
 	const ITextureFormat* TextureFormat,
 	TArray<FImage>& MipChain,
 	const FTextureBuildSettings& Settings,
-	const bool bImageHasAlphaChannel,
+	const bool bImageHasAlphaChannel, // bImageHasAlphaChannel might be unknown and invalid, however we know that if it matters to the pixel format it IS known.
 	FStringView DebugTexturePathName,
 	TArray<FCompressedImage2D>& OutMips,
 	uint32& OutNumMipsInTail,
@@ -3489,7 +3491,6 @@ static bool CompressMipChain(
 	TRACE_CPUPROFILER_EVENT_SCOPE(Texture.CompressMipChain)
 
 	// Determine ExtData (platform specific data) and NumMipsInTail.
-	// ExtData gets passed 
 	FEncodedTextureExtendedData ExtendedData;
 	FEncodedTextureDescription TextureDescription;
 	Settings.GetEncodedTextureDescription(&TextureDescription, TextureFormat, MipChain[0].SizeX, MipChain[0].SizeY, MipChain[0].NumSlices, MipChain.Num(), bImageHasAlphaChannel);
@@ -3894,6 +3895,7 @@ public:
 			return false;
 		}
 
+		TStringBuilder<32> AlphaLogs;
 		const bool bDoDetailedAlphaLogging = CVarDetailedMipAlphaLogging.GetValueOnAnyThread() != 0;
 		
 		// @todo Oodle: option to dump the Source image here
@@ -3902,17 +3904,20 @@ public:
 
 		TArray<FImage> IntermediateMipChain;
 
-		bool bSourceMipsAlphaDetected = false;
+		bool bSourceMipsExpectAlpha = false;
 		if (bDoDetailedAlphaLogging)
 		{
-			bSourceMipsAlphaDetected = FImageCore::DetectAlphaChannel(SourceMips[0]);
-			UE_LOG(LogTextureCompressor, Display, TEXT("[alpha] Source Mips: %d - %.*s"), bSourceMipsAlphaDetected, DebugTexturePathName.Len(), DebugTexturePathName.GetData());
+			AlphaLogs.Append(DebugTexturePathName);
+
+			bSourceMipsExpectAlpha = FImageCore::DetectAlphaChannel(SourceMips[0]);
+
+			AlphaLogs.Appendf(TEXT("\nSource Mips: %d / Analyzed: %d (%s)"), bSourceMipsExpectAlpha, BuildSettings.bHasTransparentAlpha, BuildSettings.bKnowAlphaTransparency ? TEXT("valid") : TEXT("INVALID"));
 		}
 
 		// allow to leave texture in sRGB in case compressor accepts other than non-F32 input source
 		// otherwise linearizing will force format to be RGBA32F
 		const bool bNeedLinearize = !TextureFormat->CanAcceptNonF32Source(BuildSettings.TextureFormatName) || AssociatedNormalSourceMips.Num() != 0;
-		if (!BuildTextureMips(SourceMips, BuildSettings, bNeedLinearize, IntermediateMipChain, DebugTexturePathName))
+		if (!BuildTextureMips(SourceMips, BuildSettings, bNeedLinearize, bDoDetailedAlphaLogging ? &AlphaLogs : nullptr, IntermediateMipChain, DebugTexturePathName))
 		{
 			return false;
 		}
@@ -3949,7 +3954,7 @@ public:
 			// use new mip filter setting from build settings
 			DefaultSettings.bUseNewMipFilter = BuildSettings.bUseNewMipFilter;
 
-			if (!BuildTextureMips(AssociatedNormalSourceMips, DefaultSettings, true, IntermediateAssociatedNormalSourceMipChain, DebugTexturePathName))
+			if (!BuildTextureMips(AssociatedNormalSourceMips, DefaultSettings, true, bDoDetailedAlphaLogging ? &AlphaLogs : nullptr, IntermediateAssociatedNormalSourceMipChain, DebugTexturePathName))
 			{
 				UE_LOG(LogTextureCompressor, Warning, TEXT("Failed to generate texture mips for composite texture [%.*s]"),
 					DebugTexturePathName.Len(),DebugTexturePathName.GetData());
@@ -3969,50 +3974,63 @@ public:
 
 			if (bDoDetailedAlphaLogging)
 			{
-				UE_LOG(LogTextureCompressor, Display, TEXT("[alpha] Associated Normal Mips: %d - %.*s"), FImageCore::DetectAlphaChannel(IntermediateMipChain[0]),
-					DebugTexturePathName.Len(), DebugTexturePathName.GetData());
+				AlphaLogs.Appendf(TEXT("\nAssociated Normal Mips: %d"), FImageCore::DetectAlphaChannel(IntermediateMipChain[0]));
 			}
 		}
 
+		bool bKnowAlphaTransparency = BuildSettings.bKnowAlphaTransparency;
+		bool bHasTransparentAlpha = BuildSettings.bHasTransparentAlpha;
+		
+		// We expect virtually all textures coming through this path to carry color information as it's scanned when the
+		// source mips are retrieved for the build. (FTextureSourceData::GetSourceMips). However if the settings don't allow
+		// for prediction off of source mips, we evaluate here.
+		if (!bKnowAlphaTransparency)
+		{
+			// We only need the alpha information if our pixel format can change as a result. This should only be
+			// the various "Auto" formats, but the safest way to check is to try both and see. Currenly the only known
+			// way to get here results in formats that do NOT change based on the presence of alpha.
+			EPixelFormat FormatWithAlpha = TextureFormat->GetEncodedPixelFormat(BuildSettings, true);
+			EPixelFormat FormatWithoutAlpha = TextureFormat->GetEncodedPixelFormat(BuildSettings, false);
+			if (FormatWithAlpha != FormatWithoutAlpha)
+			{
+				// At this time there's no known way to get here. Setting as a warning so hopefully we hear
+				// about it when it happens.
+				UE_LOG(LogTextureCompressor, Warning, TEXT("Scanning intermediate mips for alpha: %.*s"), DebugTexturePathName.Len(), DebugTexturePathName.GetData());
+
+				bool bIntermediateMipsAlphaDetected = FImageCore::DetectAlphaChannel(IntermediateMipChain[0]);
+				bHasTransparentAlpha = BuildSettings.GetTextureExpectsAlphaInPixelFormat(bIntermediateMipsAlphaDetected);
+				bKnowAlphaTransparency = true;
+			}
+			else
+			{
+				// If we don't think we need it, make sure that we use false for alpha in case something happens
+				// to look at it we have a consistent answer.
+				bHasTransparentAlpha = false;
+			}
+		}
 
-		bool bIntermediateMipsAlphaDetected = FImageCore::DetectAlphaChannel(IntermediateMipChain[0]);
 		if (bDoDetailedAlphaLogging)
 		{
-			UE_LOG(LogTextureCompressor, Display, TEXT("[alpha] Intermediate Mips: %d - %.*s"), bIntermediateMipsAlphaDetected,
-				DebugTexturePathName.Len(), DebugTexturePathName.GetData());
+			bool bIntermediateMipsAlphaDetected = FImageCore::DetectAlphaChannel(IntermediateMipChain[0]);
+			AlphaLogs.Appendf(TEXT("\nIntermediate Mips: %d"), bIntermediateMipsAlphaDetected);
+			AlphaLogs.Appendf(TEXT("\nAlpha Channel Needed: %d"), 
+				   BuildSettings.GetTextureExpectsAlphaInPixelFormat(bIntermediateMipsAlphaDetected));
 
-			if (bIntermediateMipsAlphaDetected != bSourceMipsAlphaDetected)
+			EPixelFormat FormatWithAlpha = TextureFormat->GetEncodedPixelFormat(BuildSettings, true);
+			EPixelFormat FormatWithoutAlpha = TextureFormat->GetEncodedPixelFormat(BuildSettings, false);
+
+			if (FormatWithAlpha != FormatWithoutAlpha &&
+				BuildSettings.GetTextureExpectsAlphaInPixelFormat(bIntermediateMipsAlphaDetected) != bSourceMipsExpectAlpha)
 			{
-				UE_LOG(LogTextureCompressor, Display, TEXT("[alpha] Source / Intermediate diff (force %d force no %d) - %.*s"),
+				AlphaLogs.Appendf(TEXT("\nSource / Intermediate diff (force %d force no %d)"),
 					BuildSettings.bForceAlphaChannel, 
-					BuildSettings.bForceNoAlphaChannel, 
-					DebugTexturePathName.Len(),
-					DebugTexturePathName.GetData());
+					BuildSettings.bForceNoAlphaChannel);
+				
+				UE_LOG(LogTextureCompressor, Warning, TEXT("%s"), *AlphaLogs);
 			}
-		}
-
-		// DetectAlphaChannel on the top mip of the generated mip chain. SoonTM this will use the source mip chain. Testing has
-		// shown this to be 99.9% the same, and allows us to get the pixel format earlier.
-		const bool bImageHasAlphaChannel = BuildSettings.GetTextureExpectsAlphaInPixelFormat(bIntermediateMipsAlphaDetected);
-
-		// If we know what our transparency is make sure it matches what we expect.
-		if (BuildSettings.bKnowAlphaTransparency)
-		{
-			if (BuildSettings.bHasTransparentAlpha != bImageHasAlphaChannel)
+			else
 			{
-				// We don't actually use this yet but we DO need to know when they aren't matching to hunt
-				// down bugs so we do a warning. The actual build is still good and this doesn't affect the output.
-
-				// We only actually care if it affects the output format.
-				FName ActualTextureFormat = UE::TextureBuildUtilities::TextureFormatRemovePrefixFromName(BuildSettings.TextureFormatName);
-				if (ActualTextureFormat == "AutoDXT")
-				{
-					UE_LOG(LogTextureCompressor, Warning, TEXT("Expected texture alpha didn't match reality: Expected: %s Reality: %s Texture: %.*s"),
-						BuildSettings.bHasTransparentAlpha ? TEXT("true") : TEXT("false"),
-						bImageHasAlphaChannel ? TEXT("true") : TEXT("false"),
-						DebugTexturePathName.Len(),
-						DebugTexturePathName.GetData());
-				}
+				UE_LOG(LogTextureCompressor, Display, TEXT("%s"), *AlphaLogs);
 			}
 		}
 		
@@ -4024,7 +4042,7 @@ public:
 		}
 		
 		// CompressMipChain may free IntermediateMipChain
-		bool bCompressSucceeded = CompressMipChain(TextureFormat, IntermediateMipChain, BuildSettings, bImageHasAlphaChannel, DebugTexturePathName,
+		bool bCompressSucceeded = CompressMipChain(TextureFormat, IntermediateMipChain, BuildSettings, bHasTransparentAlpha, DebugTexturePathName,
 					OutTextureMips, OutNumMipsInTail, OutExtData);
 
 		{
@@ -4077,6 +4095,7 @@ private:
 		TArray<FImage>& InSourceMipChain,
 		const FTextureBuildSettings& BuildSettings,
 		const bool bNeedLinearize,
+		TStringBuilder<32>* InAlphaLogs,
 		TArray<FImage>& OutMipChain,
 		FStringView DebugTexturePathName)
 	{
@@ -4109,8 +4128,6 @@ private:
 			}
 		}
 
-		const bool bDoDetailedAlphaLogging = CVarDetailedMipAlphaLogging.GetValueOnAnyThread() != 0;
-
 		// handling of bLongLatCubemap seems overly complicated
 		//	what it should do is convert it right at the start here
 		//	then treat it as a standard cubemap below, no special cases
@@ -4456,9 +4473,9 @@ private:
 		{
 			FImage& Image = (*pSourceMips)[MipIndex];
 
-			if (bDoDetailedAlphaLogging)
+			if (InAlphaLogs)
 			{
-				UE_LOG(LogTextureCompressor, Display, TEXT("[alpha] Pre Process: %d - %.*s"), FImageCore::DetectAlphaChannel(Image), DebugTexturePathName.Len(), DebugTexturePathName.GetData());
+				InAlphaLogs->Appendf(TEXT("\nPre Process: %d"), FImageCore::DetectAlphaChannel(Image));
 			}
 			
 			// copy mips over + processing
@@ -4487,9 +4504,9 @@ private:
 
 					GenerateTopMip(Temp, Mip, BuildSettings);
 
-					if (bDoDetailedAlphaLogging)
+					if (InAlphaLogs)
 					{
-						UE_LOG(LogTextureCompressor, Display, TEXT("[alpha] ApplyKernel: %d - %.*s"), FImageCore::DetectAlphaChannel(Mip), DebugTexturePathName.Len(), DebugTexturePathName.GetData());
+						InAlphaLogs->Appendf(TEXT("\nApplyKernel: %d"), FImageCore::DetectAlphaChannel(Mip));
 					}
 				}
 				else
@@ -4502,9 +4519,9 @@ private:
 							NormalizeMip(Mip);
 						}
 
-						if (bDoDetailedAlphaLogging)
+						if (InAlphaLogs)
 						{
-							UE_LOG(LogTextureCompressor, Display, TEXT("[alpha] Linearize: %d - %.*s"), FImageCore::DetectAlphaChannel(Mip), DebugTexturePathName.Len(), DebugTexturePathName.GetData());
+							InAlphaLogs->Appendf(TEXT("\nLinearize: %d"), FImageCore::DetectAlphaChannel(Mip));
 						}
 					}
 					else
@@ -4512,9 +4529,9 @@ private:
 						// just move Image to Mip, no copy
 						Image.Swap(Mip);
 
-						if (bDoDetailedAlphaLogging)
+						if (InAlphaLogs)
 						{
-							UE_LOG(LogTextureCompressor, Display, TEXT("[alpha] Copy: %d - %.*s"), FImageCore::DetectAlphaChannel(Mip), DebugTexturePathName.Len(), DebugTexturePathName.GetData());
+							InAlphaLogs->Appendf(TEXT("\nCopy: %d"), FImageCore::DetectAlphaChannel(Mip));
 						}
 					}
 				}
@@ -4527,18 +4544,18 @@ private:
 			{		
 				DownscaleImage(Mip, Mip, FTextureDownscaleSettings(BuildSettings));
 
-				if (bDoDetailedAlphaLogging)
+				if (InAlphaLogs)
 				{
-					UE_LOG(LogTextureCompressor, Display, TEXT("[alpha] Downscale: %d - %.*s"), FImageCore::DetectAlphaChannel(Mip), DebugTexturePathName.Len(), DebugTexturePathName.GetData());
+					InAlphaLogs->Appendf(TEXT("\nDownscale: %d"), FImageCore::DetectAlphaChannel(Mip));
 				}
 			}
 
 			// Apply color adjustments
 			AdjustImageColors(Mip, BuildSettings);
 
-			if (bDoDetailedAlphaLogging)
+			if (InAlphaLogs)
 			{
-				UE_LOG(LogTextureCompressor, Display, TEXT("[alpha] ColorAdjust: %d - %.*s"), FImageCore::DetectAlphaChannel(Mip), DebugTexturePathName.Len(), DebugTexturePathName.GetData());
+				InAlphaLogs->Appendf(TEXT("\nColorAdjust: %d"), FImageCore::DetectAlphaChannel(Mip));
 			}
 
 			if (BuildSettings.bComputeBokehAlpha)
@@ -4546,9 +4563,9 @@ private:
 				// To get the occlusion in the BokehDOF shader working for all Bokeh textures.
 				ComputeBokehAlpha(Mip);
 
-				if (bDoDetailedAlphaLogging)
+				if (InAlphaLogs)
 				{
-					UE_LOG(LogTextureCompressor, Display, TEXT("[alpha] Bokeh: %d - %.*s"), FImageCore::DetectAlphaChannel(Mip), DebugTexturePathName.Len(), DebugTexturePathName.GetData());
+					InAlphaLogs->Appendf(TEXT("\nBokeh: %d"), FImageCore::DetectAlphaChannel(Mip));
 				}
 			}
 			if (BuildSettings.bFlipGreenChannel)
@@ -4606,9 +4623,9 @@ private:
 				NumOutputMips, OutMipChain[0].SizeX, OutMipChain[0].SizeY, OutMipChain[0].NumSlices);
 		}
 
-		if (bDoDetailedAlphaLogging)
+		if (InAlphaLogs)
 		{
-			UE_LOG(LogTextureCompressor, Display, TEXT("[alpha] BeginPost: %d - %.*s"), FImageCore::DetectAlphaChannel(OutMipChain[0]), DebugTexturePathName.Len(), DebugTexturePathName.GetData());
+			InAlphaLogs->Appendf(TEXT("\nBeginPost: %d"), FImageCore::DetectAlphaChannel(OutMipChain[0]));
 		}
 
 
@@ -4648,9 +4665,9 @@ private:
 			check( !BuildSettings.bReplicateAlpha ); // cannot both be set 
 			ReplicateRedChannel(OutMipChain);
 
-			if (bDoDetailedAlphaLogging)
+			if (InAlphaLogs)
 			{
-				UE_LOG(LogTextureCompressor, Display, TEXT("[alpha] ReplicateRed: %d - %.*s"), FImageCore::DetectAlphaChannel(OutMipChain[0]), DebugTexturePathName.Len(), DebugTexturePathName.GetData());
+				InAlphaLogs->Appendf(TEXT("\nReplicateRed: %d"), FImageCore::DetectAlphaChannel(OutMipChain[0]));
 			}
 		}
 		else if (BuildSettings.bReplicateAlpha)
@@ -4663,9 +4680,9 @@ private:
 			ApplyYCoCgBlockScale(OutMipChain);
 		}
 
-		if (bDoDetailedAlphaLogging)
+		if (InAlphaLogs)
 		{
-			UE_LOG(LogTextureCompressor, Display, TEXT("[alpha] End: %d - %.*s"), FImageCore::DetectAlphaChannel(OutMipChain[0]), DebugTexturePathName.Len(), DebugTexturePathName.GetData());
+			InAlphaLogs->Appendf(TEXT("\nEnd: %d"), FImageCore::DetectAlphaChannel(OutMipChain[0]));
 		}
 
 		return true;
diff --git a/Engine/Source/Developer/TextureCompressor/Public/TextureCompressorModule.h b/Engine/Source/Developer/TextureCompressor/Public/TextureCompressorModule.h
--- a/Engine/Source/Developer/TextureCompressor/Public/TextureCompressorModule.h
+++ b/Engine/Source/Developer/TextureCompressor/Public/TextureCompressorModule.h
@@ -381,7 +381,7 @@ struct FTextureBuildSettings
 
 	// If there is a choice to be had between two formats: one with alpha and one without, this returns which
 	// one the texture expects.
-	bool GetTextureExpectsAlphaInPixelFormat(bool bInSourceMipsAlphaDetected) const
+	bool GetTextureExpectsAlphaInPixelFormat(bool bTopMipAlphaDetected) const
 	{
 		// note the order of operations! ( ForceNo takes precedence )
 		if (bForceNoAlphaChannel)
@@ -392,7 +392,7 @@ struct FTextureBuildSettings
 		{
 			return true;
 		}
-		return bInSourceMipsAlphaDetected;
+		return bTopMipAlphaDetected;
 	}
 
 	// bUnknownSourceAlphaFallback is whether we should treat the source as having alpha, in cases where the source alpha is unknown
