[AsyncCompilation/Texture]
- Add support for mid-build cancellation
- Allows to reduce wait times when exiting the application and doing GC between worlds when textures are still being built but are not needed anymore
- Time to cancel is highly dependant on the time taken by the core compression libraries since cancelation support is missing at that level
- Cancelation is inactive by default for now, can be activated with -dpcvars=Editor.AsyncTextureCompilationCancelable=true

diff --git a/Engine/Source/Developer/TextureCompressor/Private/TextureCompressorModule.cpp b/Engine/Source/Developer/TextureCompressor/Private/TextureCompressorModule.cpp
--- a/Engine/Source/Developer/TextureCompressor/Private/TextureCompressorModule.cpp
+++ b/Engine/Source/Developer/TextureCompressor/Private/TextureCompressorModule.cpp
@@ -3532,11 +3532,15 @@ static bool CompressMipChain(
 	FIntVector3 Mip0Dimensions = TextureDescription.GetMipDimensions(0);
 	int32 Mip0NumSlicesNoDepth = TextureDescription.GetNumSlices_NoDepth();
 
+	UE::Tasks::FCancellationToken* CancellationToken = UE::Tasks::FCancellationTokenScope::GetCurrentCancellationToken();
+
 	auto ProcessMips =
-		[&TextureFormat, &MipChain, &OutMips, &Mip0Dimensions, Mip0NumSlicesNoDepth, FirstMipTailIndex, MipTailCount, ExtData = ExtendedData.ExtData, &Settings, &DebugTexturePathName, bImageHasAlphaChannel](int32 MipBegin, int32 MipEnd)
+		[&TextureFormat, &MipChain, &OutMips, &Mip0Dimensions, Mip0NumSlicesNoDepth, FirstMipTailIndex, MipTailCount, ExtData = ExtendedData.ExtData, &Settings, &DebugTexturePathName, bImageHasAlphaChannel, CancellationToken](int32 MipBegin, int32 MipEnd)
 	{
 		bool bSuccess = true;
 
+		UE::Tasks::FCancellationTokenScope CancellationScope(CancellationToken);
+
 		for (int32 MipIndex = MipBegin; MipIndex < MipEnd; ++MipIndex)
 		{
 			TRACE_CPUPROFILER_EVENT_SCOPE(Texture.CompressImage);
@@ -3549,26 +3553,33 @@ static bool CompressMipChain(
 				MipsToCompress = MipTailCount;
 			}
 
-			bSuccess = bSuccess && TextureFormat->CompressImageEx(
-				&MipChain[MipIndex],
-				MipsToCompress,
-				Settings,
-				Mip0Dimensions,
-				Mip0NumSlicesNoDepth,
-				MipIndex,
-				OutMips.Num(),
-				DebugTexturePathName,
-				bImageHasAlphaChannel,
-				ExtData,
-				OutMips[MipIndex]
-			);
+			if (CancellationToken && CancellationToken->IsCanceled())
+			{
+				bSuccess = false;
+			}
+			else
+			{
+				bSuccess = bSuccess && TextureFormat->CompressImageEx(
+					&MipChain[MipIndex],
+					MipsToCompress,
+					Settings,
+					Mip0Dimensions,
+					Mip0NumSlicesNoDepth,
+					MipIndex,
+					OutMips.Num(),
+					DebugTexturePathName,
+					bImageHasAlphaChannel,
+					ExtData,
+					OutMips[MipIndex]
+				);
+			}
 
 			{
-			TRACE_CPUPROFILER_EVENT_SCOPE(Texture.CompressImage.Free);
-			// CompressImageEx can free images as it works on them
-			//	go ahead and always immediately free here for consistency
+				TRACE_CPUPROFILER_EVENT_SCOPE(Texture.CompressImage.Free);
+				// CompressImageEx can free images as it works on them
+				//	go ahead and always immediately free here for consistency
 
-			MipChain[MipIndex].FreeData(true);
+				MipChain[MipIndex].FreeData(true);
 			}
 		}
 
@@ -3616,6 +3627,11 @@ static bool CompressMipChain(
 		OutMips.Empty();
 	}
 
+	if (CancellationToken && CancellationToken->IsCanceled())
+	{
+		return false;
+	}
+
 	uint32 EndCycles = FPlatformTime::Cycles();
 	UE_LOG(LogTextureCompressor,Verbose,TEXT("Compressed %dx%dx%d %s in %fms"),
 		MipChain[0].SizeX,
@@ -3914,6 +3930,11 @@ public:
 			AlphaLogs.Appendf(TEXT("\nSource Mips: %d / Analyzed: %d (%s)"), bSourceMipsExpectAlpha, BuildSettings.bHasTransparentAlpha, BuildSettings.bKnowAlphaTransparency ? TEXT("valid") : TEXT("INVALID"));
 		}
 
+		if (UE::Tasks::FCancellationTokenScope::IsCurrentWorkCanceled())
+		{
+			return false;
+		}
+
 		// allow to leave texture in sRGB in case compressor accepts other than non-F32 input source
 		// otherwise linearizing will force format to be RGBA32F
 		const bool bNeedLinearize = !TextureFormat->CanAcceptNonF32Source(BuildSettings.TextureFormatName) || AssociatedNormalSourceMips.Num() != 0;
@@ -3922,6 +3943,11 @@ public:
 			return false;
 		}
 		
+		if (UE::Tasks::FCancellationTokenScope::IsCurrentWorkCanceled())
+		{
+			return false;
+		}
+
 		SourceMips.Empty();
 
 		// apply roughness adjustment depending on normal map variation
@@ -3956,22 +3982,42 @@ public:
 
 			if (!BuildTextureMips(AssociatedNormalSourceMips, DefaultSettings, true, bDoDetailedAlphaLogging ? &AlphaLogs : nullptr, IntermediateAssociatedNormalSourceMipChain, DebugTexturePathName))
 			{
+				if (UE::Tasks::FCancellationTokenScope::IsCurrentWorkCanceled())
+				{
+					return false;
+				}
+
 				UE_LOG(LogTextureCompressor, Warning, TEXT("Failed to generate texture mips for composite texture [%.*s]"),
 					DebugTexturePathName.Len(),DebugTexturePathName.GetData());
 
 				return false;
 			}
 
+			if (UE::Tasks::FCancellationTokenScope::IsCurrentWorkCanceled())
+			{
+				return false;
+			}
+
 			AssociatedNormalSourceMips.Empty();
 
 			if (!ApplyCompositeTextureToMips(IntermediateMipChain, IntermediateAssociatedNormalSourceMipChain, BuildSettings.CompositeTextureMode, BuildSettings.CompositePower, BuildSettings.LODBias))
 			{
+				if (UE::Tasks::FCancellationTokenScope::IsCurrentWorkCanceled())
+				{
+					return false;
+				}
+
 				UE_LOG(LogTextureCompressor, Display, TEXT("ApplyCompositeTextureToMips failed [%.*s]"),
 					DebugTexturePathName.Len(),DebugTexturePathName.GetData());
 
 				return false;
 			}
 
+			if (UE::Tasks::FCancellationTokenScope::IsCurrentWorkCanceled())
+			{
+				return false;
+			}
+
 			if (bDoDetailedAlphaLogging)
 			{
 				AlphaLogs.Appendf(TEXT("\nAssociated Normal Mips: %d"), FImageCore::DetectAlphaChannel(IntermediateMipChain[0]));
@@ -4052,6 +4098,11 @@ public:
 			IntermediateMipChain.Empty();
 		}
 
+		if (UE::Tasks::FCancellationTokenScope::IsCurrentWorkCanceled())
+		{
+			return false;
+		}
+
 		if ( bCompressSucceeded )
 		{
 			// make sure this log is done before the decode for platform preview
diff --git a/Engine/Source/Runtime/Core/Private/Tasks/TaskPrivate.cpp b/Engine/Source/Runtime/Core/Private/Tasks/TaskPrivate.cpp
--- a/Engine/Source/Runtime/Core/Private/Tasks/TaskPrivate.cpp
+++ b/Engine/Source/Runtime/Core/Private/Tasks/TaskPrivate.cpp
@@ -557,15 +557,35 @@ namespace UE::Tasks
 		static thread_local FCancellationToken* CancellationToken = nullptr;
 	}
 
+	void FCancellationTokenScope::SetToken(FCancellationToken* CancellationToken)
+	{
+		if (CancellationToken)
+		{
+			if (Private::CancellationToken != CancellationToken)
+			{
+				check(Private::CancellationToken == nullptr);
+				Private::CancellationToken = CancellationToken;
+				bHasActiveScope = true;
+			}
+		}
+	}
+
 	FCancellationTokenScope::FCancellationTokenScope(FCancellationToken& CancellationToken)
 	{
-		check(Private::CancellationToken == nullptr);
-		Private::CancellationToken = &CancellationToken;
+		SetToken(&CancellationToken);
+	}
+
+	FCancellationTokenScope::FCancellationTokenScope(FCancellationToken* CancellationToken)
+	{
+		SetToken(CancellationToken);
 	}
 
 	FCancellationTokenScope::~FCancellationTokenScope()
 	{
-		Private::CancellationToken = nullptr;
+		if (bHasActiveScope)
+		{
+			Private::CancellationToken = nullptr;
+		}
 	}
 
 	FCancellationToken* FCancellationTokenScope::GetCurrentCancellationToken()
diff --git a/Engine/Source/Runtime/Core/Public/Tasks/Task.h b/Engine/Source/Runtime/Core/Public/Tasks/Task.h
--- a/Engine/Source/Runtime/Core/Public/Tasks/Task.h
+++ b/Engine/Source/Runtime/Core/Public/Tasks/Task.h
@@ -656,6 +656,7 @@ namespace UE::Tasks
 	{
 	public:
 		CORE_API FCancellationTokenScope(FCancellationToken& CancellationToken);
+		CORE_API FCancellationTokenScope(FCancellationToken* CancellationToken);
 		CORE_API ~FCancellationTokenScope();
 
 		CORE_API static FCancellationToken* GetCurrentCancellationToken();
@@ -669,5 +670,8 @@ namespace UE::Tasks
 
 			return false;
 		}
+	private:
+		bool bHasActiveScope = false;
+		void SetToken(FCancellationToken* CancellationToken);
 	};
 }
diff --git a/Engine/Source/Runtime/Engine/Private/Tests/Texture2DTests.cpp b/Engine/Source/Runtime/Engine/Private/Tests/Texture2DTests.cpp
--- a/Engine/Source/Runtime/Engine/Private/Tests/Texture2DTests.cpp
+++ b/Engine/Source/Runtime/Engine/Private/Tests/Texture2DTests.cpp
@@ -3,6 +3,8 @@
 #include "Misc/AutomationTest.h"
 #include "Engine/Texture2D.h"
 #include "UObject/Package.h"
+#include "UObject/ReferenceChainSearch.h"
+#include "AssetCompilingManager.h"
 
 #if WITH_DEV_AUTOMATION_TESTS
 
@@ -110,6 +112,90 @@ bool FTexture2DTestLockingWhenEmpty::RunTest(const FString& Parameters)
 
 #endif //WITH_EDITORONLY_DATA
 
+#if WITH_EDITOR
+
+IMPLEMENT_SIMPLE_AUTOMATION_TEST(FTexture2DAsyncCompileCancelation, "System.Engine.Texture2D.AsyncCompileCancelation", TestFlags)
+bool FTexture2DAsyncCompileCancelation::RunTest(const FString& Parameters)
+{
+	constexpr int32 TextureSize = 2 * 1024;
+	constexpr int32 NumIteration = 10;
+
+	auto DoTest =
+		[this](bool bVirtualStreaming)
+		{
+			double CompilationTime = 0.0;
+
+			for (int32 Iteration = 0; Iteration < NumIteration; ++Iteration)
+			{
+				TRACE_CPUPROFILER_EVENT_SCOPE(FTexture2DAsyncCompileCancelation::Iteration);
+
+				UTexture2D* TestTexture = nullptr;
+		
+				{
+					TRACE_CPUPROFILER_EVENT_SCOPE(FTexture2DAsyncCompileCancelation::CreateTransient);
+					TestTexture = UTexture2D::CreateTransient(TextureSize, TextureSize);
+				}
+
+				{
+					TRACE_CPUPROFILER_EVENT_SCOPE(FTexture2DAsyncCompileCancelation::PreEditChange);
+					TestTexture->PreEditChange(nullptr);
+				}
+
+				TestTexture->VirtualTextureStreaming = bVirtualStreaming;
+
+				{
+					TRACE_CPUPROFILER_EVENT_SCOPE(FTexture2DAsyncCompileCancelation::WriteSource);
+					TestTexture->Source.Init(TextureSize, TextureSize, 1, 1, TSF_BGRA8);
+					int32* Data = (int32*)TestTexture->Source.LockMip(0);
+
+					ParallelFor(TextureSize * TextureSize, [&](int32 Index) { Data[Index] = FMath::Rand(); });
+
+					TestTexture->Source.UnlockMip(0);
+				}
+
+				{
+					TRACE_CPUPROFILER_EVENT_SCOPE(FTexture2DAsyncCompileCancelation::PostEditChange);
+					TestTexture->PostEditChange();
+				}
+	
+				// on first iteration, gather the compilation time
+				if (Iteration == 0)
+				{
+					double StartTime = FPlatformTime::Seconds();
+					FAssetCompilingManager::Get().FinishCompilationForObjects({TestTexture});
+					CompilationTime = FPlatformTime::Seconds() - StartTime;
+					continue;
+				}
+
+				TWeakObjectPtr<UTexture> WeakPtr(TestTexture);
+
+				{
+					TRACE_CPUPROFILER_EVENT_SCOPE(FTexture2DAsyncCompileCancelation::RandomSleep);
+					// On subsequent iterations, sleep a random amount of time that should span the compilation time
+					// to test various cancellation points
+					FPlatformProcess::Sleep(FMath::RandRange(0.0, CompilationTime));
+				}
+
+				{
+					TRACE_CPUPROFILER_EVENT_SCOPE(FTexture2DAsyncCompileCancelation::GC);
+					CollectGarbage(GARBAGE_COLLECTION_KEEPFLAGS);
+				}
+		
+				TestFalse(TEXT("There shouldn't be anything preventing the texture from being GCed during async compilation"), WeakPtr.IsValid());
+			}
+		};
+
+	{
+		// Normal and VT textures have different cancelation points
+		DoTest(false /*bVirtualStreaming*/);
+		DoTest(true  /*bVirtualStreaming*/);
+	}
+
+	return true;
+}
+
+#endif
+
 } // namespace Texture2DTest
 
 #endif //WITH_DEV_AUTOMATION_TESTS
diff --git a/Engine/Source/Runtime/Engine/Private/Texture.cpp b/Engine/Source/Runtime/Engine/Private/Texture.cpp
--- a/Engine/Source/Runtime/Engine/Private/Texture.cpp
+++ b/Engine/Source/Runtime/Engine/Private/Texture.cpp
@@ -21,6 +21,7 @@
 #include "IImageWrapperModule.h"
 #include "Engine/Engine.h"
 #include "Interfaces/ITargetPlatform.h"
+#include "AssetCompilingManager.h"
 #include "Engine/TextureLODSettings.h"
 #include "RenderUtils.h"
 #include "ObjectCacheContext.h"
@@ -531,6 +532,12 @@ void UTexture::BlockOnAnyAsyncBuild()
 
 bool UTexture::Modify(bool bAlwaysMarkDirty)
 {
+	if (!IsAsyncCacheComplete())
+	{
+		// try to throw away the data since we're going to invalidate it anyway
+		FAssetCompilingManager::Get().MarkCompilationAsCanceled({ this });
+	}
+
 	// Before applying any modification to the texture
 	// make sure no compilation is still ongoing.
 	BlockOnAnyAsyncBuild();
diff --git a/Engine/Source/Runtime/Engine/Private/TextureCompiler.cpp b/Engine/Source/Runtime/Engine/Private/TextureCompiler.cpp
--- a/Engine/Source/Runtime/Engine/Private/TextureCompiler.cpp
+++ b/Engine/Source/Runtime/Engine/Private/TextureCompiler.cpp
@@ -33,6 +33,13 @@ static AsyncCompilationHelpers::FAsyncCompilationStandardCVars CVarAsyncTextureS
 		}
 	));
 
+static bool GAsyncTextureCompilationCancelable = false;
+static FAutoConsoleVariableRef CVarAsyncTextureCompilationCancelable(
+	TEXT("Editor.AsyncTextureCompilationCancelable"),
+	GAsyncTextureCompilationCancelable,
+	TEXT("Whether or not to allow early cancelation of textures during async compilation."),
+	ECVF_Default);
+
 namespace TextureCompilingManagerImpl
 {
 	static FString GetLODGroupName(UTexture* Texture)
@@ -81,6 +88,11 @@ FTextureCompilingManager::FTextureCompilingManager()
 	TextureCompilingManagerImpl::EnsureInitializedCVars();
 }
 
+bool FTextureCompilingManager::IsAsyncCompilationCancelable() const
+{
+	return GAsyncTextureCompilationCancelable;
+}
+
 FName FTextureCompilingManager::GetStaticAssetTypeName()
 {
 	return TEXT("UE-Texture");
@@ -376,6 +388,50 @@ void FTextureCompilingManager::FinishCompilationForObjects(TArrayView<UObject* c
 	}
 }
 
+void FTextureCompilingManager::MarkCompilationAsCanceled(TArrayView<UObject* const> InObjects)
+{
+	if (InObjects.Num())
+	{
+		TSet<UTexture*> Textures;
+		for (UObject* Object : InObjects)
+		{
+			if (UTexture* Texture = Cast<UTexture>(Object))
+			{
+				Textures.Add(Texture);
+			}
+		}
+
+		if (Textures.Num())
+		{
+			TRACE_CPUPROFILER_EVENT_SCOPE(FTextureCompilingManager::MarkCompilationAsCanceled);
+
+			auto CancelOrMarkCanceled = [&Textures](TSet<TWeakObjectPtr<UTexture>>& Set)
+			{
+				for (auto Iterator = Set.CreateIterator(); Iterator; ++Iterator)
+				{
+					UTexture* Texture = Iterator->GetEvenIfUnreachable();
+					if (Texture && Textures.Contains(Texture))
+					{
+						UE_LOG(LogTexture, Verbose, TEXT("Canceling texture %s async compilation as requested"), *Texture->GetName());
+
+						// On success, we can remove them from the list right away.
+						// Otherwise, they are marked as canceled and will finish ASAP if the tasks support early cancellation.
+						if (Texture->TryCancelCachePlatformData())
+						{
+							Iterator.RemoveCurrent();
+						}
+					}
+				}
+			};
+
+			for (TSet<TWeakObjectPtr<UTexture>>& Bucket : RegisteredTextureBuckets)
+			{
+				CancelOrMarkCanceled(Bucket);
+			}
+		}
+	}
+}
+
 void FTextureCompilingManager::FinishCompilation(TArrayView<UTexture* const> InTextures)
 {
 	TRACE_CPUPROFILER_EVENT_SCOPE(FTextureCompilingManager::FinishCompilation);
diff --git a/Engine/Source/Runtime/Engine/Private/TextureDerivedDataTask.cpp b/Engine/Source/Runtime/Engine/Private/TextureDerivedDataTask.cpp
--- a/Engine/Source/Runtime/Engine/Private/TextureDerivedDataTask.cpp
+++ b/Engine/Source/Runtime/Engine/Private/TextureDerivedDataTask.cpp
@@ -1014,10 +1014,16 @@ static void DDC1_BuildTexture(
 		FVirtualTextureDataBuilder Builder(*DerivedData->VTData, TexturePathName, Compressor, ImageWrapper);
 		if ( ! Builder.Build(TextureData, CompositeTextureData, &InBuildSettingsPerLayer[0], true) )
 		{
-			UE_LOG(LogTexture, Warning, TEXT("VT Build failed: %s"), *TexturePathName);
 			delete DerivedData->VTData;
 			DerivedData->VTData = nullptr;
 			bSucceeded = false;
+
+			if (UE::Tasks::FCancellationTokenScope::IsCurrentWorkCanceled())
+			{
+				return;
+			}
+
+			UE_LOG(LogTexture, Warning, TEXT("VT Build failed: %s"), *TexturePathName);
 			return;		
 		}
 
@@ -1108,6 +1114,11 @@ static void DDC1_BuildTexture(
 			ExtData,
 			nullptr); // OutMetadata
 			
+		if (UE::Tasks::FCancellationTokenScope::IsCurrentWorkCanceled())
+		{
+			return;
+		}
+
 		if ( bSucceeded )
 		{
 			// BuildTexture can free the source images passed to it
@@ -1171,6 +1182,11 @@ static void DDC1_BuildTexture(
 				} // end if need to detile here.
 			} // end if decoding for PC
 
+			if (UE::Tasks::FCancellationTokenScope::IsCurrentWorkCanceled())
+			{
+				return;
+			}
+
 			if (bSucceeded)
 			{
 				check(CompressedMips.Num());
@@ -1201,6 +1217,11 @@ static void DDC1_BuildTexture(
 			check( DerivedData->Mips.Num() == 0 );
 			DerivedData->Mips.Empty();
 
+			if (UE::Tasks::FCancellationTokenScope::IsCurrentWorkCanceled())
+			{
+				return;
+			}
+
 			UE_LOG(LogTexture, Warning, TEXT("BuildTexture failed to build %s derived data for %s"), *InBuildSettingsPerLayer[0].TextureFormatName.GetPlainNameString(), *TexturePathName);
 		}
 	}
@@ -1991,6 +2012,13 @@ void FTextureCacheDerivedDataWorker::DoWork()
 	using namespace UE;
 	using namespace UE::DerivedData;
 
+	if (CancellationToken.IsCanceled())
+	{
+		return;
+	}
+
+	UE::Tasks::FCancellationTokenScope CancellationScope(CancellationToken);
+
 	TRACE_CPUPROFILER_EVENT_SCOPE(FTextureCacheDerivedDataWorker::DoWork);
 	COOK_STAT(auto Timer = TextureCookStats::TaskUsageStats.TimeSyncWork());
 
@@ -2010,6 +2038,11 @@ void FTextureCacheDerivedDataWorker::DoWork()
 		bLoadedFromDDC = true;
 	}
 
+	if (CancellationToken.IsCanceled())
+	{
+		return;
+	}
+
 	if (BuildSettingsPerLayerFetchOrBuild[0].Tiler && !bForVirtualTextureStreamingBuild)
 	{
 		if (CVarForceRetileTextures.GetValueOnAnyThread())
@@ -2030,6 +2063,11 @@ void FTextureCacheDerivedDataWorker::DoWork()
 	{
 		if (DDC1_LoadAndValidateTextureData(Texture, TextureData, CompositeTextureData, ImageWrapper))
 		{
+			if (CancellationToken.IsCanceled())
+			{
+				return;
+			}
+
 			for (int32 LayerIndex = 0; LayerIndex < BuildSettingsPerLayerFetchOrBuild.Num(); LayerIndex++)
 			{
 				if (LayerIndex < TextureData.LayerChannelMinMax.Num())
@@ -2058,6 +2096,11 @@ void FTextureCacheDerivedDataWorker::DoWork()
 					KeySuffix, bReplaceExistingDDC, RequiredMemoryEstimate, DerivedData, BytesCached, bSucceeded);
 			}
 
+			if (CancellationToken.IsCanceled())
+			{
+				return;
+			}
+
 			// TextureData may have been freed by Build, don't use it anymore
 
 			if (bInvalidVirtualTextureCompression && DerivedData->VTData)
@@ -2108,6 +2151,11 @@ void FTextureCacheDerivedDataWorker::DoWork()
 		CompositeTextureData.ReleaseMemory();
 	}
 
+	if (CancellationToken.IsCanceled())
+	{
+		return;
+	}
+
 	if ( bSucceeded )
 	{
 		// Populate the VT DDC Cache now if we're asynchronously loading to avoid too many high prio/synchronous request on the render thread
@@ -2127,7 +2175,13 @@ void FTextureCacheDerivedDataWorker::Finalize()
 	// --	always with a lighting build, as the async light/shadowmap tasks will disallow async builds. 
 	// --	if the texture compiler cvar disallows async texture compilation  "Editor.AsyncTextureCompilation 0"
 
-	TRACE_CPUPROFILER_EVENT_SCOPE(FTextureCacheDerivedDataWorker::DoWork);
+	TRACE_CPUPROFILER_EVENT_SCOPE(FTextureCacheDerivedDataWorker::Finalize);
+
+	if (CancellationToken.IsCanceled())
+	{
+		return;
+	}
+
 	COOK_STAT(auto Timer = TextureCookStats::TaskUsageStats.TimeSyncWork());
 	COOK_STAT(Timer.TrackCyclesOnly());
 
diff --git a/Engine/Source/Runtime/Engine/Private/TextureDerivedDataTask.h b/Engine/Source/Runtime/Engine/Private/TextureDerivedDataTask.h
--- a/Engine/Source/Runtime/Engine/Private/TextureDerivedDataTask.h
+++ b/Engine/Source/Runtime/Engine/Private/TextureDerivedDataTask.h
@@ -15,6 +15,7 @@
 #include "Engine/Texture2D.h"
 #include "IImageWrapperModule.h"
 #include "ImageCore.h"
+#include "TextureCompiler.h"
 #include "Misc/EnumClassFlags.h"
 #include "TextureCompressorModule.h"
 #include "TextureEncodingSettings.h"
@@ -168,6 +169,9 @@ class FTextureCacheDerivedDataWorker : public FNonAbandonableTask
 	bool bTriedAndFailed = false;
 	/** true if the derived data was pulled from DDC */
 	bool bLoadedFromDDC = false;
+	/** Will be set to true if we should abort compilation */
+	UE::Tasks::FCancellationToken CancellationToken;
+
 public:
 
 	/** Initialization constructor. */
@@ -209,6 +213,14 @@ public:
 	{
 		RETURN_QUICK_DECLARE_CYCLE_STAT(FTextureCacheDerivedDataWorker, STATGROUP_ThreadPoolAsyncTasks);
 	}
+
+	void MarkAsCanceled()
+	{
+		if (FTextureCompilingManager::Get().IsAsyncCompilationCancelable())
+		{
+			CancellationToken.Cancel();
+		}
+	}
 };
 
 struct FTextureAsyncCacheDerivedDataTask
diff --git a/Engine/Source/Runtime/Engine/Private/VT/VirtualTextureDataBuilder.cpp b/Engine/Source/Runtime/Engine/Private/VT/VirtualTextureDataBuilder.cpp
--- a/Engine/Source/Runtime/Engine/Private/VT/VirtualTextureDataBuilder.cpp
+++ b/Engine/Source/Runtime/Engine/Private/VT/VirtualTextureDataBuilder.cpp
@@ -506,6 +506,11 @@ bool FVirtualTextureDataBuilder::Build(FTextureSourceData& InSourceData, FTextur
 		}
 	}
 
+	if (UE::Tasks::FCancellationTokenScope::IsCurrentWorkCanceled())
+	{
+		return false;
+	}
+
 	// Rearrange compressed VT tiles into chunks for output
 	// Chunks contain multiple tiles: [Tiles] -> [Mips] -> [Layers]
 	bool ok = BuildChunks();
@@ -772,9 +777,16 @@ void FVirtualTextureDataBuilder::BuildBlockTiles(uint32 LayerIndex, uint32 Block
 			//	TextureFormats should disable their own internal use of TaskGraph for VT tiles if necessary
 			const bool bIsSingleThreaded = !bAllowAsync;
 
+			UE::Tasks::FCancellationToken* CancellationToken = UE::Tasks::FCancellationTokenScope::GetCurrentCancellationToken();
+
 			// Build all tiles for this mip level
 			ParallelFor(TEXT("Texture.VT.BuildTiles.PF"), NumTiles, 1, [&](int32 TileIndex)
 			{
+				if (CancellationToken && CancellationToken->IsCanceled())
+				{
+					return;
+				}
+
 				int32 TileY = TileIndex / MipBlockSizeInTilesX;
 				int32 TileX = TileIndex % MipBlockSizeInTilesX;
 
@@ -818,7 +830,14 @@ void FVirtualTextureDataBuilder::BuildBlockTiles(uint32 LayerIndex, uint32 Block
 				TArray<FImage> EmptyList;
 				uint32 NumMipsInTail, ExtData;
 				// this is the Build for Tiles to do the encode to GPU formats, with no processing
-				if (!ensure(Compressor->BuildTexture(TileImages, EmptyList, TBSettings, DebugTilePathName, CompressedMip, NumMipsInTail, ExtData, nullptr)))
+				bool bResult = Compressor->BuildTexture(TileImages, EmptyList, TBSettings, DebugTilePathName, CompressedMip, NumMipsInTail, ExtData, nullptr);
+
+				if (CancellationToken && CancellationToken->IsCanceled())
+				{
+					return;
+				}
+
+				if (!ensure(bResult))
 				{
 					bCompressionError = true;
 				}
@@ -1009,6 +1028,11 @@ void FVirtualTextureDataBuilder::BuildLayerBlocks(FSlowTask& BuildTask, uint32 L
 			bBuildTextureResult = Compressor->BuildTexture(ScaledSourceMips, ScaledCompositeMips, TBSettings, CurDebugTexturePathName, CompressedMips, NumMipsInTail, ExtData, nullptr);
 		}
 
+		if (UE::Tasks::FCancellationTokenScope::IsCurrentWorkCanceled())
+		{
+			return;
+		}
+
 		check(bBuildTextureResult);
 
 		// Get size of block from Compressor output, since it may have been padded/adjusted
diff --git a/Engine/Source/Runtime/Engine/Public/TextureCompiler.h b/Engine/Source/Runtime/Engine/Public/TextureCompiler.h
--- a/Engine/Source/Runtime/Engine/Public/TextureCompiler.h
+++ b/Engine/Source/Runtime/Engine/Public/TextureCompiler.h
@@ -25,6 +25,11 @@ public:
 	 */
 	ENGINE_API bool IsAsyncTextureCompilationEnabled() const;
 
+	/**
+	 * Returns true if the cancelation feature is currently activated.
+	 */
+	ENGINE_API bool IsAsyncCompilationCancelable() const;
+
 	/** 
 	 * Returns the number of outstanding texture compilations.
 	 */
@@ -90,6 +95,9 @@ private:
 
 	/** Handle generic finish compilation */
 	void FinishCompilationForObjects(TArrayView<UObject* const> InObjects) override;
+	
+	/** Mark compilation of the provided static meshes as canceled. */
+	void MarkCompilationAsCanceled(TArrayView<UObject* const> InObjects) override;
 
 	FName GetAssetTypeName() const override;
 	FTextFormat GetAssetNameFormat() const override;
